Electra: Added a new mp4 handler to support fragmented and non-interleaved (or interleaved with too large an interval) files to be played

#rnx

#Synced-CL 50764032

diff --git a/Engine/Plugins/Media/ElectraCodecs/Source/ElectraCodecFactory/ElectraCodecFactory.Build.cs b/Engine/Plugins/Media/ElectraCodecs/Source/ElectraCodecFactory/ElectraCodecFactory.Build.cs
--- a/Engine/Plugins/Media/ElectraCodecs/Source/ElectraCodecFactory/ElectraCodecFactory.Build.cs
+++ b/Engine/Plugins/Media/ElectraCodecs/Source/ElectraCodecFactory/ElectraCodecFactory.Build.cs
@@ -12,7 +12,8 @@ namespace UnrealBuildTool.Rules
 
 			PrivateDependencyModuleNames.AddRange(
 				new string[] {
-					"ElectraDecoders"
+					"ElectraDecoders",
+                    "ElectraBase"
 				});
 
 			PublicDependencyModuleNames.AddRange(
diff --git a/Engine/Plugins/Media/ElectraCodecs/Source/ElectraCodecFactory/Private/ElectraCodecFactoryModule.cpp b/Engine/Plugins/Media/ElectraCodecs/Source/ElectraCodecFactory/Private/ElectraCodecFactoryModule.cpp
--- a/Engine/Plugins/Media/ElectraCodecs/Source/ElectraCodecFactory/Private/ElectraCodecFactoryModule.cpp
+++ b/Engine/Plugins/Media/ElectraCodecs/Source/ElectraCodecFactory/Private/ElectraCodecFactoryModule.cpp
@@ -75,6 +75,7 @@ public:
 	}
 
 	TSharedPtr<IElectraCodecFactory, ESPMode::ThreadSafe> GetBestFactoryForFormat(TMap<FString, FVariant>& OutFormatInfo, const FString& InCodecFormat, bool bInEncoder, const TMap<FString, FVariant>& InOptions) override;
+	TSharedPtr<IElectraCodecFactory, ESPMode::ThreadSafe> GetBestDecoderFactoryForFormat(TMap<FString, FVariant>& OutFormatInfo, const Electra::FCodecTypeFormat& InCodecFormat, const TMap<FString, FVariant>& InAdditionalOptions) override;
 
 private:
 	FElectraCodecRegistry StandardDecoders;
@@ -111,6 +112,74 @@ TSharedPtr<IElectraCodecFactory, ESPMode::ThreadSafe> FElectraCodecFactoryModule
 	return BestFactory;
 }
 
+TSharedPtr<IElectraCodecFactory, ESPMode::ThreadSafe> FElectraCodecFactoryModule::GetBestDecoderFactoryForFormat(TMap<FString, FVariant>& OutFormatInfo, const Electra::FCodecTypeFormat& InCodecFormat, const TMap<FString, FVariant>& InAdditionalOptions)
+{
+	// This is a temporary helper until all decoder factories `SupportsFormat()` methods have been changed to support the `FCodecTypeFormat` as a parameter.
+	if (InCodecFormat.Type == Electra::FCodecTypeFormat::EType::Video || InCodecFormat.Type == Electra::FCodecTypeFormat::EType::Audio)
+	{
+		TMap<FString, FVariant> Params(InAdditionalOptions);
+		switch(InCodecFormat.Type)
+		{
+			case Electra::FCodecTypeFormat::EType::Video:
+			{
+				const Electra::FCodecTypeFormat::FVideo& Video(InCodecFormat.Properties.Get<Electra::FCodecTypeFormat::FVideo>());
+				Params.Add(TEXT("width"), FVariant(Video.Width));
+				Params.Add(TEXT("height"), FVariant(Video.Height));
+				if (Video.FrameRate.IsValid())
+				{
+					Params.Add(TEXT("fps"), FVariant((double)Video.FrameRate.AsDecimal()));
+					Params.Add(TEXT("fps_n"), FVariant((int64)Video.FrameRate.Numerator));
+					Params.Add(TEXT("fps_d"), FVariant((uint32)Video.FrameRate.Denominator));
+				}
+				else
+				{
+					Params.Add(TEXT("fps"), FVariant((double)0.0));
+					Params.Add(TEXT("fps_n"), FVariant((int64)0));
+					Params.Add(TEXT("fps_d"), FVariant((uint32)1));
+				}
+				break;
+			}
+			case Electra::FCodecTypeFormat::EType::Audio:
+			{
+				const Electra::FCodecTypeFormat::FAudio& Audio(InCodecFormat.Properties.Get<Electra::FCodecTypeFormat::FAudio>());
+				Params.Add(TEXT("channel_configuration"), FVariant(Audio.ChannelConfiguration));
+				Params.Add(TEXT("num_channels"), FVariant((int32) Audio.NumChannels));
+				Params.Add(TEXT("sample_rate"), FVariant((int32) Audio.SampleRate));
+				break;
+			}
+		}
+		Params.Add(TEXT("dcr"), FVariant(InCodecFormat.DCR));
+		Params.Add(TEXT("codec_4cc"), FVariant(InCodecFormat.FourCC));
+		Params.Add(TEXT("codec_name"), FVariant(InCodecFormat.RFC6381));
+		// Add children box data
+		auto GetPrintableBoxAtom = [](uint32 InAtom) -> FString
+		{
+			TCHAR tc[4];
+			tc[0] = (TCHAR) ((InAtom >> 24) & 255);
+			tc[1] = (TCHAR) ((InAtom >> 16) & 255);
+			tc[2] = (TCHAR) ((InAtom >>  8) & 255);
+			tc[3] = (TCHAR) ((InAtom >>  0) & 255);
+			for(int32 i=0;i<4; ++i)
+			{
+				tc[i] = tc[i] >= 32 && tc[i] <= 127 ? tc[i] : TCHAR(' ');
+			}
+			return FString::ConstructFromPtrSize(tc, 4);
+		};
+		for(auto& Ch : InCodecFormat.ExtraBoxes)
+		{
+			FString BoxName = FString::Printf(TEXT("$%s_box"), *GetPrintableBoxAtom(Ch.Key));
+			Params.Add(BoxName, FVariant(Ch.Value));
+		}
+		return GetBestFactoryForFormat(OutFormatInfo, InCodecFormat.RFC6381, false, Params);
+	}
+	else
+	{
+		UE_LOG(LogElectraCodecFactory, Warning, TEXT("GetBestDecoderFactoryForFormat() handles only video and audio at the moment."));
+		return nullptr;
+	}
+}
+
+
 // -----------------------------------------------------------------------------------------------------------------------------------
 void FElectraCodecRegistry::AddCodecFactory(TSharedPtr<IElectraCodecFactory, ESPMode::ThreadSafe> InCodecFactory)
 {
diff --git a/Engine/Plugins/Media/ElectraCodecs/Source/ElectraCodecFactory/Public/IElectraCodecFactoryModule.h b/Engine/Plugins/Media/ElectraCodecs/Source/ElectraCodecFactory/Public/IElectraCodecFactoryModule.h
--- a/Engine/Plugins/Media/ElectraCodecs/Source/ElectraCodecFactory/Public/IElectraCodecFactoryModule.h
+++ b/Engine/Plugins/Media/ElectraCodecs/Source/ElectraCodecFactory/Public/IElectraCodecFactoryModule.h
@@ -8,7 +8,7 @@
 #include "Modules/ModuleManager.h"
 #include "Containers/Map.h"
 #include "Misc/Variant.h"
-
+#include "CodecTypeFormat.h"
 
 class IElectraCodecFactory;
 
@@ -27,5 +27,6 @@ public:
 		return FeatureName;
 	}
 
-	virtual TSharedPtr<IElectraCodecFactory, ESPMode::ThreadSafe> GetBestFactoryForFormat(TMap<FString, FVariant>& OutFormatInfo, const FString& InCodecFormat, bool bInEncoder, const TMap<FString, FVariant>& InOptions) = 0;
+	virtual TSharedPtr<IElectraCodecFactory, ESPMode::ThreadSafe> GetBestFactoryForFormat(TMap<FString, FVariant>& OutFormatInfo, const FString& InCodecFormatRFC6831, bool bInEncoder, const TMap<FString, FVariant>& InOptions) = 0;
+	virtual TSharedPtr<IElectraCodecFactory, ESPMode::ThreadSafe> GetBestDecoderFactoryForFormat(TMap<FString, FVariant>& OutFormatInfo, const Electra::FCodecTypeFormat& InCodecFormat, const TMap<FString, FVariant>& InAdditionalOptions) = 0;
 };
diff --git a/Engine/Plugins/Media/ElectraCodecs/Source/ElectraDecoders/Private/Utils/ElectraCodecFormatUtils.cpp b/Engine/Plugins/Media/ElectraCodecs/Source/ElectraDecoders/Private/Utils/ElectraCodecFormatUtils.cpp
--- a/Engine/Plugins/Media/ElectraCodecs/Source/ElectraDecoders/Private/Utils/ElectraCodecFormatUtils.cpp
+++ b/Engine/Plugins/Media/ElectraCodecs/Source/ElectraDecoders/Private/Utils/ElectraCodecFormatUtils.cpp
@@ -71,6 +71,10 @@ static const FString CodecPrefix_MP4A(TEXT("mp4a"));
 static const FString CodecPrefix_Opus(TEXT("opus"));
 static const FString CodecPrefix_Flac(TEXT("flac"));
 
+static const FString CodecPrefix_WebVTT(TEXT("wvtt"));
+static const FString CodecPrefix_TTML(TEXT("stpp"));
+static const FString CodecPrefix_TX3G(TEXT("tx3g"));
+
 
 static inline uint32 BitReverse32(uint32 InValue)
 {
@@ -180,6 +184,33 @@ static void SetDefaultMimeTypeAndRFC6381For4CC(Electra::FCodecTypeFormat& InOut)
 				}
 				break;
 			}
+			case Make4CC('s','t','p','p'):
+			{
+				InOut.MimeType = TEXT("application/ttml+xml");
+				if (InOut.RFC6381.IsEmpty())
+				{
+					InOut.RFC6381 = TEXT("stpp.ttml.im1t");
+				}
+				break;
+			}
+			case Make4CC('w','v','t','t'):
+			{
+				InOut.MimeType = TEXT("text/vtt");
+				if (InOut.RFC6381.IsEmpty())
+				{
+					InOut.RFC6381 = TEXT("wvtt");
+				}
+				break;
+			}
+			case Make4CC('t','x','3','g'):
+			{
+				InOut.MimeType = TEXT("application/mp4");
+				if (InOut.RFC6381.IsEmpty())
+				{
+					InOut.RFC6381 = TEXT("tx3g");
+				}
+				break;
+			}
 			default:
 			{
 				if (InOut.RFC6381.IsEmpty())
@@ -947,6 +978,33 @@ static bool DetermineCodecTypeFromRFC(Electra::FCodecTypeFormat& InOut)
 			}
 			InOut.HumanReadableFormatInfo = TEXT("Free Lossless Audio Codec (FLAC)");
 		}
+		else if (InOut.RFC6381.StartsWith(CodecPrefix_WebVTT))
+		{
+			InOut.Type = Electra::FCodecTypeFormat::EType::Subtitle;
+			if (!InOut.Properties.IsType<Electra::FCodecTypeFormat::FSubtitle>())
+			{
+				InOut.Properties.Emplace<Electra::FCodecTypeFormat::FSubtitle>();
+			}
+			InOut.HumanReadableFormatInfo = TEXT("WebVTT");
+		}
+		else if (InOut.RFC6381.StartsWith(CodecPrefix_TTML))
+		{
+			InOut.Type = Electra::FCodecTypeFormat::EType::Subtitle;
+			if (!InOut.Properties.IsType<Electra::FCodecTypeFormat::FSubtitle>())
+			{
+				InOut.Properties.Emplace<Electra::FCodecTypeFormat::FSubtitle>();
+			}
+			InOut.HumanReadableFormatInfo = TEXT("Timed Text (TTML)");
+		}
+		else if (InOut.RFC6381.StartsWith(CodecPrefix_TX3G))
+		{
+			InOut.Type = Electra::FCodecTypeFormat::EType::Subtitle;
+			if (!InOut.Properties.IsType<Electra::FCodecTypeFormat::FSubtitle>())
+			{
+				InOut.Properties.Emplace<Electra::FCodecTypeFormat::FSubtitle>();
+			}
+			InOut.HumanReadableFormatInfo = TEXT("Timed Text (TX3G)");
+		}
 	}
 	return bOk;
 }
@@ -988,7 +1046,7 @@ static int32 DetermineCodecTypeFromDCR_avcc(Electra::FCodecTypeFormat& InOut)
 			vid.Profile.Constraints = dcr.GetProfileCompatibility();
 			InOut.CSD = dcr.GetCodecSpecificData();
 			int32 FourCC = InOut.CSD.IsEmpty() ? Make4CC('a','v','c','3') : Make4CC('a','v','c','1');
-			if (InOut.FourCC && InOut.FourCC != FourCC)
+			if (InOut.FourCC && InOut.FourCC != FourCC && FourCC == Make4CC('a','v','c','3'))
 			{
 				UE_LOG(LogElectraFormatUtils, Log, TEXT("FourCC was set as %08x but identifed as %08x from AVCDecoderConfigurationRecord"), InOut.FourCC, FourCC);
 			}
@@ -1075,8 +1133,8 @@ static int32 DetermineCodecTypeFromDCR_hvcc(Electra::FCodecTypeFormat& InOut)
 			vid.Profile.Constraints = dcr.GetGeneralConstraintIndicatorFlags();
 			vid.Profile.CompatibilityFlags = dcr.GetGeneralProfileCompatibilityFlags();
 			InOut.CSD = dcr.GetCodecSpecificData();
-			int32 FourCC = !InOut.CSD.IsEmpty() ? Make4CC('h','v','c','1') : Make4CC('h','e','v','1');
-			if (InOut.FourCC && InOut.FourCC != FourCC)
+			int32 FourCC = InOut.CSD.IsEmpty() ? Make4CC('h','e','v','1') : Make4CC('h','v','c','1');
+			if (InOut.FourCC && InOut.FourCC != FourCC && FourCC == Make4CC('h','e','v','1'))
 			{
 				UE_LOG(LogElectraFormatUtils, Log, TEXT("FourCC was set as %08x but identifed as %08x from HEVCDecoderConfigurationRecord"), InOut.FourCC, FourCC);
 			}
@@ -1516,6 +1574,79 @@ static int32 DetermineCodecTypeFromDCR_Flac(Electra::FCodecTypeFormat& InOut)
 	return 0;
 }
 
+static int32 DetermineCodecTypeFromDCR_tx3g(Electra::FCodecTypeFormat& InOut)
+{
+	const TArray<uint8>* tx3g_Box = InOut.ExtraBoxes.Find(Make4CC('t','x','3','g'));
+	if (tx3g_Box)
+	{
+		if (!InOut.Properties.IsType<Electra::FCodecTypeFormat::FSubtitle>())
+		{
+			InOut.Properties.Emplace<Electra::FCodecTypeFormat::FSubtitle>();
+		}
+		InOut.DCR = *tx3g_Box;
+		InOut.CSD = *tx3g_Box;
+		InOut.RFC6381 = TEXT("tx3g");
+		InOut.Type = Electra::FCodecTypeFormat::EType::Subtitle;
+		uint32 FourCC = Make4CC('t','x','3','g');
+		if (InOut.FourCC && InOut.FourCC != FourCC)
+		{
+			UE_LOG(LogElectraFormatUtils, Log, TEXT("FourCC was set as %08x but identifed as %08x"), InOut.FourCC, FourCC);
+		}
+		InOut.FourCC = FourCC;
+		return 1;
+	}
+	return 0;
+}
+
+static int32 DetermineCodecTypeFromDCR_wvtt(Electra::FCodecTypeFormat& InOut)
+{
+	const TArray<uint8>* wvtt_Box = InOut.ExtraBoxes.Find(Make4CC('w','v','t','t'));
+	if (wvtt_Box)
+	{
+		if (!InOut.Properties.IsType<Electra::FCodecTypeFormat::FSubtitle>())
+		{
+			InOut.Properties.Emplace<Electra::FCodecTypeFormat::FSubtitle>();
+		}
+		InOut.DCR = *wvtt_Box;
+		InOut.CSD = *wvtt_Box;
+		InOut.RFC6381 = TEXT("wvtt");
+		InOut.Type = Electra::FCodecTypeFormat::EType::Subtitle;
+		uint32 FourCC = Make4CC('w','v','t','t');
+		if (InOut.FourCC && InOut.FourCC != FourCC)
+		{
+			UE_LOG(LogElectraFormatUtils, Log, TEXT("FourCC was set as %08x but identifed as %08x"), InOut.FourCC, FourCC);
+		}
+		InOut.FourCC = FourCC;
+		return 1;
+	}
+	return 0;
+}
+
+static int32 DetermineCodecTypeFromDCR_stpp(Electra::FCodecTypeFormat& InOut)
+{
+	const TArray<uint8>* stpp_Box = InOut.ExtraBoxes.Find(Make4CC('s','t','p','p'));
+	if (stpp_Box)
+	{
+		if (!InOut.Properties.IsType<Electra::FCodecTypeFormat::FSubtitle>())
+		{
+			InOut.Properties.Emplace<Electra::FCodecTypeFormat::FSubtitle>();
+		}
+		// TTML does not have codec specific data
+		//InOut.DCR = *stpp_Box;
+		//InOut.CSD = *stpp_Box;
+		InOut.RFC6381 = TEXT("stpp.ttml.im1t");
+		InOut.Type = Electra::FCodecTypeFormat::EType::Subtitle;
+		uint32 FourCC = Make4CC('s','t','p','p');
+		if (InOut.FourCC && InOut.FourCC != FourCC)
+		{
+			UE_LOG(LogElectraFormatUtils, Log, TEXT("FourCC was set as %08x but identifed as %08x"), InOut.FourCC, FourCC);
+		}
+		InOut.FourCC = FourCC;
+		return 1;
+	}
+	return 0;
+}
+
 
 
 static bool DetermineCodecTypeFromDCR(Electra::FCodecTypeFormat& InOut)
@@ -1545,6 +1676,16 @@ static bool DetermineCodecTypeFromDCR(Electra::FCodecTypeFormat& InOut)
 	Result = DetermineCodecTypeFromDCR_Flac(InOut);
 	if (Result) return /*Result < 0 ? false :*/ true;
 
+	// TX3G
+	Result = DetermineCodecTypeFromDCR_tx3g(InOut);
+	if (Result) return /*Result < 0 ? false :*/ true;
+	// WebVTT
+	Result = DetermineCodecTypeFromDCR_wvtt(InOut);
+	if (Result) return /*Result < 0 ? false :*/ true;
+	// TTML
+	Result = DetermineCodecTypeFromDCR_stpp(InOut);
+	if (Result) return /*Result < 0 ? false :*/ true;
+
 	return true;
 }
 
@@ -1563,6 +1704,46 @@ static bool ProcessExtraBoxes(Electra::FCodecTypeFormat& InOut)
 	return true;
 }
 
+static bool DetermineKeyframeMode(Electra::FCodecTypeFormat& InOut)
+{
+	// Check if the keyframe mode was already set in the process.
+	// If it was then this is already the definitive answer.
+	// Otherwise we apply the most common defaults.
+	if (InOut.KeyframeMode == Electra::FCodecTypeFormat::EKeyframeMode::Unknown)
+	{
+		// All audio and subtitles we handle are keyframe only. Timecode is keyframe as well.
+		if (InOut.Type == Electra::FCodecTypeFormat::EType::Audio || InOut.Type == Electra::FCodecTypeFormat::EType::Subtitle || InOut.Type == Electra::FCodecTypeFormat::EType::Timecode)
+		{
+			InOut.KeyframeMode = Electra::FCodecTypeFormat::EKeyframeMode::OnlyKeyframes;
+		}
+		else if (InOut.Type == Electra::FCodecTypeFormat::EType::Video)
+		{
+			// For video we start out assuming everything is a keyframe. Of the codecs we know there are
+			// several where this is not true so these need to be checked for.
+			switch(InOut.FourCC)
+			{
+				case Make4CC('a','v','c','1'):
+				case Make4CC('a','v','c','3'):
+				case Make4CC('h','e','v','1'):
+				case Make4CC('h','v','c','1'):
+				case Make4CC('v','p','0','8'):
+				case Make4CC('v','p','0','9'):
+				case Make4CC('a','v','0','1'):
+				{
+					InOut.KeyframeMode = Electra::FCodecTypeFormat::EKeyframeMode::DeltaFrames;
+					break;
+				}
+				default:
+				{
+					InOut.KeyframeMode = Electra::FCodecTypeFormat::EKeyframeMode::OnlyKeyframes;
+					break;
+				}
+			}
+		}
+	}
+	return true;
+}
+
 }
 
 
@@ -1573,6 +1754,6 @@ bool FElectraCodecFormatUtilsModularFeature::PrepareCodecTypeFormat(Electra::FCo
 	RemapLegacyFourCC(InOutCodecTypeFormat);
 	SetDefaultMimeTypeAndRFC6381For4CC(InOutCodecTypeFormat);
 	DetermineCodecTypeFromMimeType(InOutCodecTypeFormat);
-	bool bOk = DetermineCodecTypeFromRFC(InOutCodecTypeFormat) && DetermineCodecTypeFromDCR(InOutCodecTypeFormat) && ProcessExtraBoxes(InOutCodecTypeFormat);
+	bool bOk = DetermineCodecTypeFromRFC(InOutCodecTypeFormat) && DetermineCodecTypeFromDCR(InOutCodecTypeFormat) && ProcessExtraBoxes(InOutCodecTypeFormat) && DetermineKeyframeMode(InOutCodecTypeFormat);
 	return bOk;
 }
diff --git a/Engine/Plugins/Media/ElectraPlayer/Source/ElectraPlayerRuntime/ElectraPlayerRuntime.Build.cs b/Engine/Plugins/Media/ElectraPlayer/Source/ElectraPlayerRuntime/ElectraPlayerRuntime.Build.cs
--- a/Engine/Plugins/Media/ElectraPlayer/Source/ElectraPlayerRuntime/ElectraPlayerRuntime.Build.cs
+++ b/Engine/Plugins/Media/ElectraPlayer/Source/ElectraPlayerRuntime/ElectraPlayerRuntime.Build.cs
@@ -33,6 +33,7 @@ namespace UnrealBuildTool.Rules
 					"ElectraCDM",
 					"ElectraSubtitles",
                     "MP4Utilities",
+                    "MP4Boxes",
 					"XmlParser",
 					"SoundTouchZ",
                     "HTTP"
diff --git a/Engine/Plugins/Media/ElectraPlayer/Source/ElectraPlayerRuntime/Private/Runtime/HTTP/HTTPManager.cpp b/Engine/Plugins/Media/ElectraPlayer/Source/ElectraPlayerRuntime/Private/Runtime/HTTP/HTTPManager.cpp
--- a/Engine/Plugins/Media/ElectraPlayer/Source/ElectraPlayerRuntime/Private/Runtime/HTTP/HTTPManager.cpp
+++ b/Engine/Plugins/Media/ElectraPlayer/Source/ElectraPlayerRuntime/Private/Runtime/HTTP/HTTPManager.cpp
@@ -484,6 +484,7 @@ namespace Electra
 			//Request->bAutoRemoveWhenComplete = bAutoRemoveWhenComplete;
 			RequestsToAdd.Enqueue(Request);
 			RequestChangesEvent.Signal();
+			HTTPStreamHandler->TriggerWorkSignal();
 		}
 	}
 
@@ -503,6 +504,7 @@ namespace Electra
 			RequestsToRemove.Enqueue(MoveTemp(Remove));
 			RequestChangesEvent.Signal();
 			Lock.Unlock();
+			HTTPStreamHandler->TriggerWorkSignal();
 		}
 		else
 		{
@@ -514,6 +516,7 @@ namespace Electra
 			RequestsToRemove.Enqueue(MoveTemp(Remove));
 			RequestChangesEvent.Signal();
 			Lock.Unlock();
+			HTTPStreamHandler->TriggerWorkSignal();
 			WaitingEvent->Wait();
 		}
 	}
@@ -735,7 +738,7 @@ namespace Electra
 						}
 						else
 						{
-							Request->ConnectionInfo.StatusInfo.ErrorDetail.SetError(UEMEDIA_ERROR_INTERNAL).SetFacility(Facility::EFacility::HTTPReader).SetMessage("HTTP request failed on ProcessRequest()");
+							Request->ConnectionInfo.StatusInfo.ErrorDetail.SetError(UEMEDIA_ERROR_INTERNAL).SetFacility(Facility::EFacility::HTTPReader).SetMessage(TEXT("HTTP request failed on ProcessRequest()"));
 							RequestsCompleted.Enqueue(Request);
 						}
 					}
@@ -1467,7 +1470,7 @@ namespace Electra
 									}
 									else
 									{
-										ci.StatusInfo.ErrorDetail.SetError(UEMEDIA_ERROR_INTERNAL).SetFacility(Facility::EFacility::HTTPReader).SetMessage("HTTP sub request failed on AddRequest()");
+										ci.StatusInfo.ErrorDetail.SetError(UEMEDIA_ERROR_INTERNAL).SetFacility(Facility::EFacility::HTTPReader).SetMessage(TEXT("HTTP sub request failed on AddRequest()"));
 									}
 								}
 								else
@@ -1518,7 +1521,7 @@ namespace Electra
 			// Go through the notions of this being a network request.
 			Request->ConnectionInfo.bIsConnected = true;
 			Request->ConnectionInfo.bHaveResponseHeaders = true;
-			Request->ConnectionInfo.ContentType = "application/octet-stream";
+			Request->ConnectionInfo.ContentType = TEXT("application/octet-stream");
 			Request->ConnectionInfo.EffectiveURL = Request->Parameters.URL;
 			Request->ConnectionInfo.HTTPVersionReceived = 11;
 			Request->ConnectionInfo.bIsChunked = false;
diff --git a/Engine/Plugins/Media/ElectraPlayer/Source/ElectraPlayerRuntime/Private/Runtime/Player/AdaptiveStreamingPlayer.cpp b/Engine/Plugins/Media/ElectraPlayer/Source/ElectraPlayerRuntime/Private/Runtime/Player/AdaptiveStreamingPlayer.cpp
--- a/Engine/Plugins/Media/ElectraPlayer/Source/ElectraPlayerRuntime/Private/Runtime/Player/AdaptiveStreamingPlayer.cpp
+++ b/Engine/Plugins/Media/ElectraPlayer/Source/ElectraPlayerRuntime/Private/Runtime/Player/AdaptiveStreamingPlayer.cpp
@@ -171,6 +171,47 @@ FErrorDetail FAdaptiveStreamingPlayer::GetError() const
 
 
 
+void FAdaptiveStreamingPlayer::GetStreamBufferUtilization(FAccessUnitBufferInfo& OutInfo, const FAccessUnit* InForThis)
+{
+	FScopeLock lock(&DataBuffersCriticalSection);
+	bool bFirst = true;
+	if (ActiveDataOutputBuffers.IsValid())
+	{
+		TSharedPtrTS<FMultiTrackAccessUnitBuffer> MultiBuffer = ActiveDataOutputBuffers->GetBuffer(InForThis->ESType);
+		TSharedPtrTS<FAccessUnitBuffer> TrackBuffer = MultiBuffer ? MultiBuffer->FindBufferFor(InForThis->BufferSourceInfo) : nullptr;
+		if (TrackBuffer)
+		{
+			TrackBuffer->GetStats(OutInfo);
+			bFirst = false;
+		}
+	}
+	for(int32 i=0,iMax=NextDataBuffers.Num(); i<iMax; ++i)
+	{
+		TSharedPtrTS<FMultiTrackAccessUnitBuffer> MultiBuffer = NextDataBuffers[i]->GetBuffer(InForThis->ESType);
+		TSharedPtrTS<FAccessUnitBuffer> TrackBuffer = MultiBuffer ? MultiBuffer->FindBufferFor(InForThis->BufferSourceInfo) : nullptr;
+		if (TrackBuffer)
+		{
+			if (bFirst)
+			{
+				TrackBuffer->GetStats(OutInfo);
+				bFirst = false;
+			}
+			else
+			{
+				FAccessUnitBufferInfo info;
+				TrackBuffer->GetStats(info);
+				OutInfo.PlayableDuration += info.PlayableDuration;
+				OutInfo.CurrentMemInUse += info.CurrentMemInUse;
+				OutInfo.NumCurrentAccessUnits += info.NumCurrentAccessUnits;
+				OutInfo.bEndOfData |= info.bEndOfData;
+				if (i == iMax-1)
+				{
+					OutInfo.bLastPushWasBlocked = info.bLastPushWasBlocked;
+				}
+			}
+		}
+	}
+}
 
 void FAdaptiveStreamingPlayer::GetStreamBufferUtilization(FAccessUnitBufferInfo& OutInfo, EStreamType BufferType)
 {
@@ -1340,25 +1381,12 @@ bool FAdaptiveStreamingPlayer::OnFragmentAccessUnitReceived(FAccessUnit* InAcces
 	// Get the current buffer stats. Although they are updated regularly once per tick this would not reflect the current
 	// state until the the next tick, which would allow the limits to be ignored as new access units arrive quickly.
 	FAccessUnitBufferInfo bufferInfo;
-	const FAccessUnitBuffer::FConfiguration* BufferConfig = nullptr;
-	if (InAccessUnit->ESType == EStreamType::Video)
-	{
-		GetStreamBufferUtilization(bufferInfo, InAccessUnit->ESType);
-		BufferConfig = &PlayerConfig.StreamBufferConfigVideo;
-	}
-	else if (InAccessUnit->ESType == EStreamType::Audio)
-	{
-		GetStreamBufferUtilization(bufferInfo, InAccessUnit->ESType);
-		BufferConfig = &PlayerConfig.StreamBufferConfigAudio;
-	}
-	else if (InAccessUnit->ESType == EStreamType::Subtitle)
-	{
-		GetStreamBufferUtilization(bufferInfo, InAccessUnit->ESType);
-		BufferConfig = &PlayerConfig.StreamBufferConfigText;
-	}
+	GetStreamBufferUtilization(bufferInfo, InAccessUnit);
+	const FAccessUnitBuffer::FConfiguration* BufferConfig = InAccessUnit->ESType == EStreamType::Video ? &PlayerConfig.StreamBufferConfigVideo :
+															InAccessUnit->ESType == EStreamType::Audio ? &PlayerConfig.StreamBufferConfigAudio :
+															InAccessUnit->ESType == EStreamType::Subtitle ? &PlayerConfig.StreamBufferConfigText : nullptr;
 	FAccessUnitBuffer::FExternalBufferInfo TotalUtilization;
 	TotalUtilization.Duration = bufferInfo.PlayableDuration;
-
 	// Try to push the data into the receiving buffer
 	return Buffer->Push(InAccessUnit, BufferConfig, &TotalUtilization);
 }
diff --git a/Engine/Plugins/Media/ElectraPlayer/Source/ElectraPlayerRuntime/Private/Runtime/Player/AdaptiveStreamingPlayerInternal.h b/Engine/Plugins/Media/ElectraPlayer/Source/ElectraPlayerRuntime/Private/Runtime/Player/AdaptiveStreamingPlayerInternal.h
--- a/Engine/Plugins/Media/ElectraPlayer/Source/ElectraPlayerRuntime/Private/Runtime/Player/AdaptiveStreamingPlayerInternal.h
+++ b/Engine/Plugins/Media/ElectraPlayer/Source/ElectraPlayerRuntime/Private/Runtime/Player/AdaptiveStreamingPlayerInternal.h
@@ -1985,7 +1985,7 @@ private:
 		{
 			ActiveRequest.Reset();
 		}
-		
+
 		bool HasPendingSeek()
 		{
 			FScopeLock lock(&Lock);
@@ -2266,6 +2266,7 @@ private:
 	bool IsSeamlessBufferSwitchPossible(EStreamType InStreamType, const TSharedPtrTS<FStreamDataBuffers>& InFromStreamBuffers);
 
 	void GetStreamBufferUtilization(FAccessUnitBufferInfo& OutInfo, EStreamType BufferType);
+	void GetStreamBufferUtilization(FAccessUnitBufferInfo& OutInfo, const FAccessUnit* InForThis);
 
 	void FeedDecoder(EStreamType Type, IAccessUnitBufferInterface* Decoder, bool bHandleUnderrun);
 	void ClearAllDecoderEODs();
diff --git a/Engine/Plugins/Media/ElectraPlayer/Source/ElectraPlayerRuntime/Private/Runtime/Player/AdaptiveStreamingPlayerManifest.cpp b/Engine/Plugins/Media/ElectraPlayer/Source/ElectraPlayerRuntime/Private/Runtime/Player/AdaptiveStreamingPlayerManifest.cpp
--- a/Engine/Plugins/Media/ElectraPlayer/Source/ElectraPlayerRuntime/Private/Runtime/Player/AdaptiveStreamingPlayerManifest.cpp
+++ b/Engine/Plugins/Media/ElectraPlayer/Source/ElectraPlayerRuntime/Private/Runtime/Player/AdaptiveStreamingPlayerManifest.cpp
@@ -6,6 +6,7 @@
 #include "Player/AdaptivePlayerOptionKeynames.h"
 #include "Player/HLS/PlaylistHandlerHLS.h"
 #include "Player/mp4/PlaylistReaderMP4.h"
+#include "Player/isobmff/PlaylistReaderISOBMFF.h"
 #include "Player/DASH/PlaylistReaderDASH.h"
 #include "Player/mkv/PlaylistReaderMKV.h"
 #include "Player/mpegaudio/PlaylistReaderMPEGAudio.h"
@@ -16,6 +17,18 @@
 #include "MediaStreamMetadata.h"
 
 
+/*********************************************************************************************************************/
+
+#include "HAL/IConsoleManager.h"
+static TAutoConsoleVariable<bool> CVarElectraUseLegacyMP4PLayer(
+	TEXT("Electra.UseLegacyMP4Player"),
+	false,
+	TEXT("Use old sequentially reading mp4 player instead of new player.\n")
+	TEXT(" true: use old player, false use new player"),
+	ECVF_Default);
+
+/*********************************************************************************************************************/
+
 namespace Electra
 {
 
@@ -231,7 +244,14 @@ void FAdaptiveStreamingPlayer::InternalLoadManifest(const FString& InURL, const
 			}
 			else if (mimeType == Playlist::MIMETypeMP4 || mimeType == Playlist::MIMETypeMP4A || mimeType == Playlist::MIMETypeQuickTime)
 			{
-				ManifestReader = IPlaylistReaderMP4::Create(this);
+				if (CVarElectraUseLegacyMP4PLayer.GetValueOnAnyThread())
+				{
+					ManifestReader = IPlaylistReaderMP4::Create(this);
+				}
+				else
+				{
+					ManifestReader = IPlaylistReaderISOBMFF::Create(this);
+				}
 			}
 			else if (mimeType == Playlist::MIMETypeDASH)
 			{
diff --git a/Engine/Plugins/Media/ElectraPlayer/Source/ElectraPlayerRuntime/Private/Runtime/Player/isobmff/OptionKeynamesISOBMFF.h b/Engine/Plugins/Media/ElectraPlayer/Source/ElectraPlayerRuntime/Private/Runtime/Player/isobmff/OptionKeynamesISOBMFF.h
new file mode 100644
--- /dev/null
+++ b/Engine/Plugins/Media/ElectraPlayer/Source/ElectraPlayerRuntime/Private/Runtime/Player/isobmff/OptionKeynamesISOBMFF.h
@@ -0,0 +1,23 @@
+// Copyright Epic Games, Inc. All Rights Reserved.
+
+#pragma once
+
+#include "CoreTypes.h"
+#include "Containers/UnrealString.h"
+#include "UObject/NameTypes.h"
+
+namespace Electra
+{
+namespace ISOBMFF
+{
+
+const FName OptionKeyISOBMFFLoadConnectTimeout(TEXT("isobmff_connection_timeout"));				//!< (FTimeValue) value specifying connection timeout
+const FName OptionKeyISOBMFFLoadNoDataTimeout(TEXT("isobmff_nodata_timeout"));					//!< (FTimeValue) value specifying no-data timeout
+
+const FName OptionKeyISOBMFFTruncateToShortestTrack(TEXT("isobmff_trunc_to_shortest"));			//!< (bool) whether or not to truncate the movie duration to the shortest track
+
+}
+
+} // namespace Electra
+
+
diff --git a/Engine/Plugins/Media/ElectraPlayer/Source/ElectraPlayerRuntime/Private/Runtime/Player/isobmff/PlaylistISOBMFF.cpp b/Engine/Plugins/Media/ElectraPlayer/Source/ElectraPlayerRuntime/Private/Runtime/Player/isobmff/PlaylistISOBMFF.cpp
new file mode 100644
--- /dev/null
+++ b/Engine/Plugins/Media/ElectraPlayer/Source/ElectraPlayerRuntime/Private/Runtime/Player/isobmff/PlaylistISOBMFF.cpp
@@ -0,0 +1,1710 @@
+// Copyright Epic Games, Inc. All Rights Reserved.
+
+#include "PlaylistISOBMFF.h"
+#include "OptionKeynamesISOBMFF.h"
+#include "StreamReaderISOBMFF.h"
+#include "PlayerCore.h"
+#include "Player/Manifest.h"
+#include "Player/PlayerStreamFilter.h"
+#include "Player/AdaptivePlayerOptionKeynames.h"
+#include "Player/AdaptivePlayerMetadataKeynames.h"
+#include "Player/DRM/DRMManager.h"
+#include "Utilities/TimeUtilities.h"
+#include "Utilities/BCP47-Helpers.h"
+#include "Utilities/MP4Helpers.h"
+#include "ElectraPlayerPrivate.h"
+#include "HTTP/HTTPManager.h"
+#include "Misc/FrameRate.h"
+#include "Misc/FrameNumber.h"
+#include "Containers/Set.h"
+#include "HAL/LowLevelMemTracker.h"
+#include "Stats/Stats.h"
+
+/*********************************************************************************************************************/
+
+#include "HAL/IConsoleManager.h"
+static TAutoConsoleVariable<int32> CVarElectraISOBMFFPlayerForceDuration(
+	TEXT("Electra.ISOBMFFPlayer.ForceDuration"),
+	0,
+	TEXT("Set movie duration to the shortest or longest duration tracks in the file.\n")
+	TEXT(" 0: keep whatever is in effect\n 1: use shortest duration\n 2: use longest duration"),
+	ECVF_Default);
+
+/*********************************************************************************************************************/
+
+#define ERRCODE_MANIFEST_ISOBMFF_NO_PLAYABLE_STREAMS		1
+#define ERRCODE_MANIFEST_ISOBMFF_STARTSEGMENT_NOT_FOUND		2
+#define ERRCODE_MANIFEST_ISOBMFF_DRM_ERROR					3
+
+
+#define DEFAULT_TRUNCATE_MOVIE_TO_SHORTEST_TRACK false
+
+
+DECLARE_CYCLE_STAT(TEXT("FPlayPeriodISOBMFF::FindSegment"), STAT_ElectraPlayer_ISOBMFF_FindSegment, STATGROUP_ElectraPlayer);
+
+namespace Electra
+{
+
+//-----------------------------------------------------------------------------
+/**
+ * CTOR
+ */
+FManifestISOBMFFInternal::FManifestISOBMFFInternal(IPlayerSessionServices* InPlayerSessionServices)
+	: PlayerSessionServices(InPlayerSessionServices)
+{
+}
+
+
+//-----------------------------------------------------------------------------
+/**
+ * DTOR
+ */
+FManifestISOBMFFInternal::~FManifestISOBMFFInternal()
+{
+}
+
+
+//-----------------------------------------------------------------------------
+/**
+ * Builds the internal manifest from the mp4's moov box.
+ */
+FErrorDetail FManifestISOBMFFInternal::Build(TArray<TSharedPtr<MP4Boxes::FMP4BoxBase>>&& InParsedRootBoxes, const FString& InURL,
+											 TSharedPtr<FMP4DataLoader, ESPMode::ThreadSafe> InDataLoader, FMP4DataLoader::FCancellationCheckDelegate InCancelCheck)
+{
+	MediaAsset = MakeSharedTS<FTimelineAssetISOBMFF>();
+	FErrorDetail Result = MediaAsset->Build(PlayerSessionServices, MoveTemp(InParsedRootBoxes), InURL, InDataLoader, InCancelCheck);
+	FTimeRange PlaybackRange = GetPlaybackRange(IManifest::EPlaybackRangeType::TemporaryPlaystartRange);
+	DefaultStartTime = PlaybackRange.Start;
+	DefaultEndTime = PlaybackRange.End;
+	return Result;
+}
+
+
+//-----------------------------------------------------------------------------
+/**
+ * Logs a message.
+ */
+void FManifestISOBMFFInternal::LogMessage(IInfoLog::ELevel Level, const FString& Message)
+{
+	if (PlayerSessionServices)
+	{
+		PlayerSessionServices->PostLog(Facility::EFacility::ISOBMFFPlaylist, Level, Message);
+	}
+}
+
+
+
+//-----------------------------------------------------------------------------
+/**
+ * Returns track metadata.
+ */
+void FManifestISOBMFFInternal::GetTrackMetadata(TArray<FTrackMetadata>& OutMetadata, EStreamType StreamType) const
+{
+	if (MediaAsset.IsValid())
+	{
+		MediaAsset->GetMetaData(OutMetadata, StreamType);
+	}
+}
+
+
+
+//-----------------------------------------------------------------------------
+/**
+ * Returns the playback range on the timeline, which is a subset of the total
+ * time range. This may be set through manifest internal means or by URL fragment
+ * parameters where permissable (eg. example.mp4#t=22,50).
+ * If start or end are not specified they will be set to invalid.
+ */
+FTimeRange FManifestISOBMFFInternal::GetPlaybackRange(EPlaybackRangeType InRangeType) const
+{
+	FTimeRange FromTo;
+
+	// We are interested in the 't' or 'r' fragment value here.
+	FString Time;
+	for(auto& Fragment : URLFragmentComponents)
+	{
+		if ((InRangeType == IManifest::EPlaybackRangeType::TemporaryPlaystartRange && Fragment.Name.Equals(TEXT("t"))) ||
+			(InRangeType == IManifest::EPlaybackRangeType::LockedPlaybackRange && Fragment.Name.Equals(TEXT("r"))))
+		{
+			Time = Fragment.Value;
+		}
+	}
+	if (!Time.IsEmpty())
+	{
+		FTimeRange TotalRange = GetTotalTimeRange();
+		TArray<FString> TimeRange;
+		const TCHAR* const TimeDelimiter = TEXT(",");
+		Time.ParseIntoArray(TimeRange, TimeDelimiter, false);
+		if (TimeRange.Num() && !TimeRange[0].IsEmpty())
+		{
+			RFC2326::ParseNPTTime(FromTo.Start, TimeRange[0]);
+		}
+		if (TimeRange.Num() > 1 && !TimeRange[1].IsEmpty())
+		{
+			RFC2326::ParseNPTTime(FromTo.End, TimeRange[1]);
+		}
+		// Need to clamp this into the total time range to prevent any issues.
+		if (FromTo.Start.IsValid() && TotalRange.Start.IsValid() && FromTo.Start < TotalRange.Start)
+		{
+			FromTo.Start = TotalRange.Start;
+		}
+		if (FromTo.End.IsValid() && TotalRange.End.IsValid() && FromTo.End > TotalRange.End)
+		{
+			FromTo.End = TotalRange.End;
+		}
+	}
+	return FromTo;
+}
+
+
+//-----------------------------------------------------------------------------
+/**
+ * Returns the minimum duration of content that must be buffered up before playback
+ * will begin. This is an arbitrary choice that could be controlled by a 'pdin' box.
+ */
+FTimeValue FManifestISOBMFFInternal::GetMinBufferTime() const
+{
+	// NOTE: This could come from a 'pdin' (progressive download information) box, but those are rarely, if ever, set by any tool.
+	return FTimeValue().SetFromSeconds(1.0);
+}
+
+TSharedPtrTS<IProducerReferenceTimeInfo> FManifestISOBMFFInternal::GetProducerReferenceTimeInfo(int64 ID) const
+{
+	return nullptr;
+}
+
+TRangeSet<double> FManifestISOBMFFInternal::GetPossiblePlaybackRates(EPlayRateType InForType) const
+{
+	TRangeSet<double> Ranges;
+//	Ranges.Add(TRange<double>{1.0}); // normal (real-time) playback rate
+	Ranges.Add(TRange<double>::Inclusive(0.5, 4.0));
+	Ranges.Add(TRange<double>{0.0}); // and pause
+	return Ranges;
+}
+
+
+//-----------------------------------------------------------------------------
+/**
+ * Creates an instance of a stream reader to stream from the mp4 file.
+ */
+IStreamReader* FManifestISOBMFFInternal::CreateStreamReaderHandler()
+{
+	FStreamReaderISOBMFF* Reader = new FStreamReaderISOBMFF;
+	if (MediaAsset)
+	{
+		switch(MediaAsset->GetLoaderType())
+		{
+			default:
+			case FTimelineAssetISOBMFF::ELoaderType::Chunks:
+			{
+				Reader->SetLoaderType(FStreamSegmentRequestISOBMFF::ELoaderType::Chunks);
+				break;
+			}
+			case FTimelineAssetISOBMFF::ELoaderType::Muxed:
+			{
+				Reader->SetLoaderType(FStreamSegmentRequestISOBMFF::ELoaderType::Muxed);
+				break;
+			}
+			case FTimelineAssetISOBMFF::ELoaderType::Fragmented:
+			{
+				Reader->SetLoaderType(FStreamSegmentRequestISOBMFF::ELoaderType::Fragmented);
+				break;
+			}
+		}
+	}
+	return Reader;
+}
+
+
+//-----------------------------------------------------------------------------
+/**
+ * Returns the playback period for the given time.
+ */
+IManifest::FResult FManifestISOBMFFInternal::FindPlayPeriod(TSharedPtrTS<IPlayPeriod>& OutPlayPeriod, const FPlayStartPosition& StartPosition, ESearchType SearchType)
+{
+	if (MediaAsset.IsValid() && StartPosition.Time.IsValid() && StartPosition.Time < MediaAsset->GetDuration())
+	{
+		OutPlayPeriod = MakeSharedTS<FPlayPeriodISOBMFF>(MediaAsset);
+		return IManifest::FResult(IManifest::FResult::EType::Found);
+	}
+	else
+	{
+		return IManifest::FResult(IManifest::FResult::EType::PastEOS);
+	}
+}
+
+IManifest::FResult FManifestISOBMFFInternal::FindNextPlayPeriod(TSharedPtrTS<IPlayPeriod>& OutPlayPeriod, TSharedPtrTS<const IStreamSegment> CurrentSegment)
+{
+	// There is no following period.
+	return IManifest::FResult(IManifest::FResult::EType::PastEOS);
+}
+
+
+
+
+
+//-----------------------------------------------------------------------------
+/**
+ * Constructs a playback period.
+ */
+FManifestISOBMFFInternal::FPlayPeriodISOBMFF::FPlayPeriodISOBMFF(TSharedPtrTS<FManifestISOBMFFInternal::FTimelineAssetISOBMFF> InMediaAsset)
+	: MediaAsset(InMediaAsset)
+	, CurrentReadyState(IManifest::IPlayPeriod::EReadyState::NotLoaded)
+	, LoaderType(InMediaAsset->GetLoaderType())
+{
+	AggregateMinChunkDuration = FTimespan::FromSeconds(1.0);
+	MuxedSegmentDuration = FTimespan::FromSeconds(2.0);
+}
+
+
+//-----------------------------------------------------------------------------
+/**
+ * Destroys a playback period.
+ */
+FManifestISOBMFFInternal::FPlayPeriodISOBMFF::~FPlayPeriodISOBMFF()
+{
+}
+
+
+//-----------------------------------------------------------------------------
+/**
+ * Sets stream playback preferences for this playback period.
+ */
+void FManifestISOBMFFInternal::FPlayPeriodISOBMFF::SetStreamPreferences(EStreamType ForStreamType, const FStreamSelectionAttributes& StreamAttributes)
+{
+	if (ForStreamType == EStreamType::Video)
+	{
+		VideoPreferences = StreamAttributes;
+	}
+	else if (ForStreamType == EStreamType::Audio)
+	{
+		AudioPreferences = StreamAttributes;
+	}
+	else if (ForStreamType == EStreamType::Subtitle)
+	{
+		SubtitlePreferences = StreamAttributes;
+	}
+}
+
+
+//-----------------------------------------------------------------------------
+/**
+ * Returns the starting bitrate.
+ *
+ * This is merely informational and not strictly required.
+ * If fetching of the moov box provided us with the total size of the mp4 file
+ * we will use that divided by the duration.
+ */
+int64 FManifestISOBMFFInternal::FPlayPeriodISOBMFF::GetDefaultStartingBitrate() const
+{
+	return 2000000;
+}
+
+//-----------------------------------------------------------------------------
+/**
+ * Returns the ready state of this playback period.
+ */
+IManifest::IPlayPeriod::EReadyState FManifestISOBMFFInternal::FPlayPeriodISOBMFF::GetReadyState()
+{
+	return CurrentReadyState;
+}
+
+
+void FManifestISOBMFFInternal::FPlayPeriodISOBMFF::Load()
+{
+	CurrentReadyState = IManifest::IPlayPeriod::EReadyState::Loaded;
+}
+
+//-----------------------------------------------------------------------------
+/**
+ * Prepares the playback period for playback.
+ * With an mp4 file we are actually always ready for playback, but we say we're not
+ * one time to get here with any possible options.
+ */
+void FManifestISOBMFFInternal::FPlayPeriodISOBMFF::PrepareForPlay()
+{
+	SelectedVideoMetadata.Reset();
+	SelectedAudioMetadata.Reset();
+	SelectedSubtitleMetadata.Reset();
+	VideoBufferSourceInfo.Reset();
+	AudioBufferSourceInfo.Reset();
+	SubtitleBufferSourceInfo.Reset();
+	SelectInitialStream(EStreamType::Video);
+	SelectInitialStream(EStreamType::Audio);
+	SelectInitialStream(EStreamType::Subtitle);
+	CurrentReadyState = IManifest::IPlayPeriod::EReadyState::IsReady;
+}
+
+
+TSharedPtrTS<FBufferSourceInfo> FManifestISOBMFFInternal::FPlayPeriodISOBMFF::GetSelectedStreamBufferSourceInfo(EStreamType StreamType)
+{
+	return StreamType == EStreamType::Video ? VideoBufferSourceInfo :
+		   StreamType == EStreamType::Audio ? AudioBufferSourceInfo :
+		   StreamType == EStreamType::Subtitle ? SubtitleBufferSourceInfo : nullptr;
+}
+
+FString FManifestISOBMFFInternal::FPlayPeriodISOBMFF::GetSelectedAdaptationSetID(EStreamType StreamType)
+{
+	switch(StreamType)
+	{
+		case EStreamType::Video:
+		{
+			return SelectedVideoMetadata.IsValid() ? SelectedVideoMetadata->ID : FString();
+		}
+		case EStreamType::Audio:
+		{
+			return SelectedAudioMetadata.IsValid() ? SelectedAudioMetadata->ID : FString();
+		}
+		case EStreamType::Subtitle:
+		{
+			return SelectedSubtitleMetadata.IsValid() ? SelectedSubtitleMetadata->ID : FString();
+		}
+		default:
+		{
+			return FString();
+		}
+	}
+}
+
+
+IManifest::IPlayPeriod::ETrackChangeResult FManifestISOBMFFInternal::FPlayPeriodISOBMFF::ChangeTrackStreamPreference(EStreamType StreamType, const FStreamSelectionAttributes& StreamAttributes)
+{
+	if (LoaderType == FTimelineAssetISOBMFF::ELoaderType::Chunks || LoaderType == FTimelineAssetISOBMFF::ELoaderType::Fragmented)
+	{
+		return IManifest::IPlayPeriod::ETrackChangeResult::NewPeriodNeeded;
+	}
+	else if (LoaderType == FTimelineAssetISOBMFF::ELoaderType::Muxed)
+	{
+		TSharedPtrTS<FTrackMetadata> Metadata = SelectMetadataForAttributes(StreamType, StreamAttributes);
+		if (Metadata.IsValid())
+		{
+			if (StreamType == EStreamType::Video)
+			{
+				if (!(SelectedVideoMetadata.IsValid() && Metadata->Equals(*SelectedVideoMetadata)))
+				{
+					SelectedVideoMetadata = Metadata;
+					MakeBufferSourceInfoFromMetadata(StreamType, VideoBufferSourceInfo, SelectedVideoMetadata);
+					return IManifest::IPlayPeriod::ETrackChangeResult::Changed;
+				}
+			}
+			else if (StreamType == EStreamType::Audio)
+			{
+				if (!(SelectedAudioMetadata.IsValid() && Metadata->Equals(*SelectedAudioMetadata)))
+				{
+					SelectedAudioMetadata = Metadata;
+					MakeBufferSourceInfoFromMetadata(StreamType, AudioBufferSourceInfo, SelectedAudioMetadata);
+					return IManifest::IPlayPeriod::ETrackChangeResult::Changed;
+				}
+			}
+			else if (StreamType == EStreamType::Subtitle)
+			{
+				if (!(SelectedSubtitleMetadata.IsValid() && Metadata->Equals(*SelectedSubtitleMetadata)))
+				{
+					SelectedSubtitleMetadata = Metadata;
+					MakeBufferSourceInfoFromMetadata(StreamType, SubtitleBufferSourceInfo, SelectedSubtitleMetadata);
+					return IManifest::IPlayPeriod::ETrackChangeResult::Changed;
+				}
+			}
+		}
+		return IManifest::IPlayPeriod::ETrackChangeResult::NotChanged;
+	}
+	else
+	{
+		unimplemented();
+		return IManifest::IPlayPeriod::ETrackChangeResult::StartOver;
+	}
+}
+
+void FManifestISOBMFFInternal::FPlayPeriodISOBMFF::SelectInitialStream(EStreamType StreamType)
+{
+	if (StreamType == EStreamType::Video)
+	{
+		SelectedVideoMetadata = SelectMetadataForAttributes(StreamType, VideoPreferences);
+		MakeBufferSourceInfoFromMetadata(StreamType, VideoBufferSourceInfo, SelectedVideoMetadata);
+	}
+	else if (StreamType == EStreamType::Audio)
+	{
+		SelectedAudioMetadata = SelectMetadataForAttributes(StreamType, AudioPreferences);
+		MakeBufferSourceInfoFromMetadata(StreamType, AudioBufferSourceInfo, SelectedAudioMetadata);
+	}
+	else if (StreamType == EStreamType::Subtitle)
+	{
+		SelectedSubtitleMetadata = SelectMetadataForAttributes(StreamType, SubtitlePreferences);
+		MakeBufferSourceInfoFromMetadata(StreamType, SubtitleBufferSourceInfo, SelectedSubtitleMetadata);
+	}
+}
+
+TSharedPtrTS<FTrackMetadata> FManifestISOBMFFInternal::FPlayPeriodISOBMFF::SelectMetadataForAttributes(EStreamType StreamType, const FStreamSelectionAttributes& InAttributes)
+{
+	TSharedPtrTS<FTimelineAssetISOBMFF> Asset = MediaAsset.Pin();
+	if (Asset.IsValid())
+	{
+		TArray<FTrackMetadata> Metadata;
+		Asset->GetMetaData(Metadata, StreamType);
+		// Is there a fixed index to be used?
+		if (InAttributes.OverrideIndex.IsSet() && InAttributes.OverrideIndex.GetValue() >= 0 && InAttributes.OverrideIndex.GetValue() < Metadata.Num())
+		{
+			// Use this.
+			return MakeSharedTS<FTrackMetadata>(Metadata[InAttributes.OverrideIndex.GetValue()]);
+		}
+		if (Metadata.Num())
+		{
+			// We do not look at the 'kind' or 'codec' here, only the language.
+			// Set the first track as default in case we do not find the one we're looking for.
+			if (InAttributes.Language_RFC4647.IsSet())
+			{
+				TArray<int32> CandidateIndices;
+				TArray<BCP47::FLanguageTag> CandList;
+				for(auto &Meta : Metadata)
+				{
+					CandList.Emplace(Meta.LanguageTagRFC5646);
+				}
+				CandidateIndices = BCP47::FindExtendedFilteringMatch(CandList, InAttributes.Language_RFC4647.GetValue());
+				return MakeSharedTS<FTrackMetadata>(CandidateIndices.Num() ? Metadata[CandidateIndices[0]] : Metadata[0]);
+			}
+			// Subtitle tracks are not selected by default. If there is no explicit selection asked for we ignore them.
+			if (StreamType != EStreamType::Subtitle)
+			{
+				return MakeSharedTS<FTrackMetadata>(Metadata[0]);
+			}
+		}
+	}
+	return nullptr;
+}
+
+void FManifestISOBMFFInternal::FPlayPeriodISOBMFF::MakeBufferSourceInfoFromMetadata(EStreamType StreamType, TSharedPtrTS<FBufferSourceInfo>& OutBufferSourceInfo, TSharedPtrTS<FTrackMetadata> InMetadata)
+{
+	if (InMetadata.IsValid())
+	{
+		OutBufferSourceInfo = MakeSharedTS<FBufferSourceInfo>();
+		OutBufferSourceInfo->Kind = InMetadata->Kind;
+		OutBufferSourceInfo->LanguageTag = InMetadata->LanguageTagRFC5646;
+		OutBufferSourceInfo->Codec = InMetadata->HighestBandwidthCodec.GetCodecName();
+		TSharedPtrTS<FTimelineAssetISOBMFF> Asset = MediaAsset.Pin();
+		OutBufferSourceInfo->PeriodID = Asset->GetUniqueIdentifier();
+		OutBufferSourceInfo->PeriodAdaptationSetID = Asset->GetUniqueIdentifier() + TEXT(".") + InMetadata->ID;
+		TArray<FTrackMetadata> Metadata;
+		Asset->GetMetaData(Metadata, StreamType);
+		for(int32 i=0; i<Metadata.Num(); ++i)
+		{
+			if (Metadata[i].Equals(*InMetadata))
+			{
+				OutBufferSourceInfo->HardIndex = i;
+				break;
+			}
+		}
+	}
+}
+
+
+
+//-----------------------------------------------------------------------------
+/**
+ * Returns the timeline media asset. We have a weak pointer to it only to
+ * prevent any cyclic locks, so we need to lock it first.
+ */
+TSharedPtrTS<ITimelineMediaAsset> FManifestISOBMFFInternal::FPlayPeriodISOBMFF::GetMediaAsset() const
+{
+	return MediaAsset.Pin();
+}
+
+
+//-----------------------------------------------------------------------------
+/**
+ * Selects a particular stream (== internal track ID) for playback.
+ */
+void FManifestISOBMFFInternal::FPlayPeriodISOBMFF::SelectStream(const FString& AdaptationSetID, const FString& RepresentationID, int32 QualityIndex, int32 MaxQualityIndex)
+{
+	// No-op.
+}
+
+void FManifestISOBMFFInternal::FPlayPeriodISOBMFF::TriggerInitSegmentPreload(const TArray<FInitSegmentPreload>& InitSegmentsToPreload)
+{
+	// No-op.
+}
+
+void FManifestISOBMFFInternal::FPlayPeriodISOBMFF::AggregateChunkSamples(TSharedPtr<IStreamSegment, ESPMode::ThreadSafe> InOutRequest, const FTimespan& InMinDuration)
+{
+/*
+	const FStreamSegmentRequestISOBMFF* Request = static_cast<const FStreamSegmentRequestISOBMFF*>(InOutRequest.Get());
+	auto It = Request->Track.TrackIterator->Clone();
+	int64 ChunkEndOffset = It->GetSampleFileOffset() + Request->Track.TrackChunkSampleInfo.SizeInBytes;
+*/
+}
+
+void FManifestISOBMFFInternal::FPlayPeriodISOBMFF::AdjustMuxedDownload(TSharedPtr<IStreamSegment, ESPMode::ThreadSafe> InOutRequest, const FTimespan& InMinDuration, EAdjustMuxType InType)
+{
+	FStreamSegmentRequestISOBMFF* Request = static_cast<FStreamSegmentRequestISOBMFF*>(InOutRequest.Get());
+
+	auto GetFirstSelectedIterator = [&]() -> TSharedPtr<MP4Boxes::FMP4Track::FIterator, ESPMode::ThreadSafe>
+	{
+		for(int32 i=0,iMax=Request->DependentStreams.Num(); i<iMax; ++i)
+		{
+			if (Request->DependentStreams[i]->bWasSelectedAtStart && Request->DependentStreams[i]->Track.TrackIterator)
+			{
+				Request->ReportingIndexInMux = i;
+				return Request->DependentStreams[i]->Track.TrackIterator->Clone();
+			}
+		}
+		return nullptr;
+	};
+
+	auto GetNextActiveIterator = [&]() -> TSharedPtr<MP4Boxes::FMP4Track::FIterator, ESPMode::ThreadSafe>
+	{
+		for(int32 i=0,iMax=Request->DependentStreams.Num(); i<iMax; ++i)
+		{
+			if (Request->DependentStreams[i]->Track.TrackIterator)
+			{
+				Request->ReportingIndexInMux = i;
+				return Request->DependentStreams[i]->Track.TrackIterator->Clone();
+			}
+		}
+		return nullptr;
+	};
+
+	if (InType == EAdjustMuxType::FirstStart)
+	{
+		auto PrimaryIterator = GetFirstSelectedIterator();
+		// If there is no track with an an iterator we're done.
+		if (!PrimaryIterator)
+		{
+			Request->bIsLastSegment = true;
+			return;
+		}
+
+		int64 FirstSampleFileOffset = PrimaryIterator->GetSampleFileOffset();
+		FTimespan AccumulatedPrimaryDuration(0);
+		bool bReachedEnd = false;
+		while(AccumulatedPrimaryDuration < InMinDuration)
+		{
+			AccumulatedPrimaryDuration += PrimaryIterator->GetDurationAsTimespan();
+			if (!PrimaryIterator->Next())
+			{
+				bReachedEnd = true;
+				break;
+			}
+		}
+		Request->Bitrate = 0;
+		Request->SegmentDuration.SetFromTimespan(AccumulatedPrimaryDuration);
+		Request->PrimarySampleFileOffsetStart = FirstSampleFileOffset;
+		Request->PrimarySampleFileOffsetEnd = bReachedEnd ? -1 : PrimaryIterator->GetSampleFileOffset() + PrimaryIterator->GetSampleSize();
+		// Adjust tracks accordingly.
+		for(int32 nTrk=0,nTrkMax=Request->DependentStreams.Num(); nTrk<nTrkMax; ++nTrk)
+		{
+			// Add the bitrate of the track even if it's already EOS.
+			Request->Bitrate += Request->DependentStreams[nTrk]->Bitrate;
+
+			const TSharedPtrTS<FStreamSegmentRequestISOBMFF>& DepStr(Request->DependentStreams[nTrk]);
+			// Track already at EOS?
+			if (!DepStr->Track.TrackIterator)
+			{
+				continue;
+			}
+			// If the track has the selected sample position before the first sample of the primary track
+			// and this track is not also a selected track or is a subtitle track (for which it may be possible
+			// that the first selected sample is way too old) we step forward.
+			TSharedPtr<MP4Boxes::FMP4Track::FIterator, ESPMode::ThreadSafe>& TkIt(DepStr->Track.TrackIterator);
+			if (!DepStr->bWasSelectedAtStart || DepStr->StreamType == EStreamType::Subtitle)
+			{
+				while(TkIt->GetSampleFileOffset() < FirstSampleFileOffset)
+				{
+					// If advancing results in an EOS we reset the iterator to indicate this.
+					if (!TkIt->Next())
+					{
+						DepStr->Track.TrackIterator.Reset();
+						DepStr->bIsLastSegment = true;
+						break;
+					}
+				}
+			}
+			else
+			{
+				// Otherwise, if this sample comes before the one from the primary track we
+				// need to load that one and thus have to adjust the file position backward.
+				int64 SmpPos = TkIt->GetSampleFileOffset();
+				if (SmpPos < Request->PrimarySampleFileOffsetStart)
+				{
+					Request->PrimarySampleFileOffsetStart = SmpPos;
+				}
+			}
+			if (TkIt)
+			{
+				TkIt->GetCurrentChunkRemainingSampleInfo(DepStr->Track.TrackChunkSampleInfo);
+				DepStr->SegmentDuration.SetFromND(DepStr->Track.TrackChunkSampleInfo.Duration, TkIt->GetTimescale());
+				auto DTS = TkIt->GetEffectiveDTS();
+				DepStr->FirstDTS.SetFromND(DTS.GetNumerator(), DTS.GetDenominator());
+				DepStr->bIsLastSegment = DepStr->Track.TrackChunkSampleInfo.bIsLastChunk;
+			}
+		}
+	}
+	else if (InType == EAdjustMuxType::Continuation)
+	{
+		Request->bIsLastSegment = true;
+		// Get the file offset where we have to continue reading from.
+		// This will be the smallest sample file offset of all remaining tracks.
+		Request->PrimarySampleFileOffsetStart = TNumericLimits<int64>::Max();
+		for(auto &tkit : Request->DependentStreams)
+		{
+			if (tkit->Track.TrackIterator)
+			{
+				Request->bIsLastSegment = false;
+				const int64 TkSmpOff = tkit->Track.TrackIterator->GetSampleFileOffset();
+				if (TkSmpOff < Request->PrimarySampleFileOffsetStart)
+				{
+					Request->PrimarySampleFileOffsetStart = TkSmpOff;
+				}
+				tkit->Track.TrackIterator->GetCurrentChunkRemainingSampleInfo(tkit->Track.TrackChunkSampleInfo);
+				tkit->SegmentDuration.SetFromND(tkit->Track.TrackChunkSampleInfo.Duration, tkit->Track.TrackIterator->GetTimescale());
+				auto DTS = tkit->Track.TrackIterator->GetEffectiveDTS();
+				tkit->FirstDTS.SetFromND(DTS.GetNumerator(), DTS.GetDenominator());
+				tkit->bIsLastSegment = tkit->Track.TrackChunkSampleInfo.bIsLastChunk;
+			}
+		}
+		auto PrimaryIterator = GetNextActiveIterator();
+		if (!PrimaryIterator)
+		{
+			Request->PrimarySampleFileOffsetStart = -1;
+			Request->PrimarySampleFileOffsetEnd = -1;
+			Request->bIsLastSegment = true;
+			return;
+		}
+		FTimespan AccumulatedPrimaryDuration(0);
+		bool bReachedEnd = false;
+		while(AccumulatedPrimaryDuration < InMinDuration)
+		{
+			AccumulatedPrimaryDuration += PrimaryIterator->GetDurationAsTimespan();
+			if (!PrimaryIterator->Next())
+			{
+				bReachedEnd = true;
+				break;
+			}
+		}
+		Request->SegmentDuration.SetFromTimespan(AccumulatedPrimaryDuration);
+		Request->PrimarySampleFileOffsetEnd = bReachedEnd ? -1 : PrimaryIterator->GetSampleFileOffset() + PrimaryIterator->GetSampleSize();
+	}
+	else //if (InType == EAdjustMuxType::Retry)
+	{
+		Request->bIsLastSegment = true;
+		// Get the file offset where we have to continue reading from.
+		// This will be the smallest sample file offset of all remaining tracks.
+		Request->PrimarySampleFileOffsetStart = TNumericLimits<int64>::Max();
+		for(auto &tkit : Request->DependentStreams)
+		{
+			if (tkit->Track.TrackIterator)
+			{
+				Request->bIsLastSegment = false;
+				const int64 TkSmpOff = tkit->Track.TrackIterator->GetSampleFileOffset();
+				if (TkSmpOff < Request->PrimarySampleFileOffsetStart)
+				{
+					Request->PrimarySampleFileOffsetStart = TkSmpOff;
+				}
+				tkit->Track.TrackIterator->GetCurrentChunkRemainingSampleInfo(tkit->Track.TrackChunkSampleInfo);
+				tkit->SegmentDuration.SetFromND(tkit->Track.TrackChunkSampleInfo.Duration, tkit->Track.TrackIterator->GetTimescale());
+				auto DTS = tkit->Track.TrackIterator->GetEffectiveDTS();
+				tkit->FirstDTS.SetFromND(DTS.GetNumerator(), DTS.GetDenominator());
+				tkit->bIsLastSegment = tkit->Track.TrackChunkSampleInfo.bIsLastChunk;
+			}
+		}
+	}
+}
+
+//-----------------------------------------------------------------------------
+/**
+ * Creates the starting segment request to start playback with.
+ */
+IManifest::FResult FManifestISOBMFFInternal::FPlayPeriodISOBMFF::GetStartingSegment(TSharedPtrTS<IStreamSegment>& OutSegment, const FPlayerSequenceState& InSequenceState, const FPlayStartPosition& StartPosition, ESearchType SearchType)
+{
+	TSharedPtrTS<FTimelineAssetISOBMFF> ma = MediaAsset.Pin();
+	if (!ma)
+	{
+		return IManifest::FResult(IManifest::FResult::EType::NotFound);
+	}
+	SCOPE_CYCLE_COUNTER(STAT_ElectraPlayer_ISOBMFF_FindSegment);
+	CSV_SCOPED_TIMING_STAT(ElectraPlayer, ISOBMFF_FindSegment);
+
+	// Frame accurate seek required?
+	const bool bFrameAccurateSearch = StartPosition.Options.bFrameAccuracy;
+	FTimeValue PlayRangeEnd = StartPosition.Options.PlaybackRange.End;
+	check(PlayRangeEnd.IsValid());
+	if (PlayRangeEnd > ma->GetDuration())
+	{
+		PlayRangeEnd = ma->GetDuration();
+	}
+
+	FTimespan SearchTime = StartPosition.Time.GetAsTimespan();
+
+
+	auto SetError = [&]() -> IManifest::FResult
+	{
+		IManifest::FResult res(IManifest::FResult::EType::NotFound);
+		res.SetErrorDetail(FErrorDetail().SetError(UEMEDIA_ERROR_DETAIL)
+							.SetFacility(Facility::EFacility::MP4Playlist)
+							.SetCode(ERRCODE_MANIFEST_ISOBMFF_STARTSEGMENT_NOT_FOUND)
+							.SetMessage(FString::Printf(TEXT("Could not find start segment for time %lld"), (long long int)StartPosition.Time.GetAsHNS())));
+		return res;
+	};
+
+	auto SetupTrackRequest = [&](TSharedPtrTS<FAdaptationSetISOBMFF> InAdapt, EStreamType InStreamType) -> TSharedPtr<FStreamSegmentRequestISOBMFF, ESPMode::ThreadSafe>
+	{
+		if (!InAdapt)
+		{
+			return nullptr;
+		}
+		auto Representation = InAdapt->GetRepresentation();
+		auto TrkInf = Representation->GetTrackInfo();
+		auto TrkIt = TrkInf->MP4Track->CreateIteratorAtKeyframe(SearchTime, FTimespan::Zero());
+		auto Req = MakeShared<FStreamSegmentRequestISOBMFF, ESPMode::ThreadSafe>();
+		Req->StreamType = InStreamType;
+		Req->bWasSelectedAtStart = true;
+		Req->Bitrate = Representation->GetBitrate();
+		Req->AdaptationSetID = InAdapt->GetUniqueIdentifier();
+		Req->AdaptationSetTrackIndex = InAdapt->GetInternalTrackIndex();
+		Req->RepresentationID = Representation->GetUniqueIdentifier();
+		Req->LanguageTag = InAdapt->GetLanguageTag();
+		Req->KindOfTrack = InAdapt->GetTrackKind();
+		Req->Track.TrackInfo = TrkInf;
+		if (TrkIt)
+		{
+			TrkIt->GetCurrentChunkRemainingSampleInfo(Req->Track.TrackChunkSampleInfo);
+			Req->Track.TrackIterator = TrkIt;
+			Req->SegmentDuration.SetFromND(Req->Track.TrackChunkSampleInfo.Duration, TrkIt->GetTimescale());
+			auto DTS = TrkIt->GetEffectiveDTS();
+			Req->FirstDTS.SetFromND(DTS.GetNumerator(), DTS.GetDenominator());
+			Req->bIsLastSegment = Req->Track.TrackChunkSampleInfo.bIsLastChunk;
+		}
+		else
+		{
+			Req->bIsLastSegment = true;
+		}
+		return Req;
+	};
+
+	// Depending on how we need to load the data we have to locate the first sample to start with differently.
+	if (LoaderType == FTimelineAssetISOBMFF::ELoaderType::Chunks || LoaderType == FTimelineAssetISOBMFF::ELoaderType::Fragmented)
+	{
+		TSharedPtr<FStreamSegmentRequestISOBMFF, ESPMode::ThreadSafe> Requests[3];
+
+		// Is there a selected video track?
+		if (SelectedVideoMetadata && !SelectedVideoMetadata->ID.IsEmpty())
+		{
+			Requests[0] = SetupTrackRequest(ma->GetAdaptationSetByTypeAndID(EStreamType::Video, SelectedVideoMetadata->ID), EStreamType::Video);
+			if (!Requests[0] || !Requests[0]->Track.TrackIterator)
+			{
+				return SetError();
+			}
+			// When not using frame accurate seeking we set the searchtime for audio and subtitles to the time of the video keyframe
+			// since we will start there and need the audio to be there, too.
+			if (!bFrameAccurateSearch)
+			{
+				SearchTime = Requests[0]->Track.TrackIterator->GetEffectivePTS().GetAsTimespan();
+			}
+		}
+		// Is there a selected audio track?
+		if (SelectedAudioMetadata && !SelectedAudioMetadata->ID.IsEmpty())
+		{
+			Requests[1] = SetupTrackRequest(ma->GetAdaptationSetByTypeAndID(EStreamType::Audio, SelectedAudioMetadata->ID), EStreamType::Audio);
+			if (!Requests[1] || !Requests[1]->Track.TrackIterator)
+			{
+				return SetError();
+			}
+		}
+		// Is there a selected subtitle track?
+		if (SelectedSubtitleMetadata && !SelectedSubtitleMetadata->ID.IsEmpty())
+		{
+			// For subtitles if we can't create a request we do not fail.
+			Requests[2] = SetupTrackRequest(ma->GetAdaptationSetByTypeAndID(EStreamType::Subtitle, SelectedSubtitleMetadata->ID), EStreamType::Subtitle);
+			if (!Requests[2] || !Requests[2]->Track.TrackIterator)
+			{
+				Requests[2].Reset();
+			}
+		}
+
+		TSharedPtr<FStreamSegmentRequestISOBMFF, ESPMode::ThreadSafe> StartRequest = MakeShared<FStreamSegmentRequestISOBMFF, ESPMode::ThreadSafe>();
+		StartRequest->MediaAsset = ma;
+		StartRequest->LoaderType = LoaderType == FTimelineAssetISOBMFF::ELoaderType::Chunks ? FStreamSegmentRequestISOBMFF::ELoaderType::Chunks : FStreamSegmentRequestISOBMFF::ELoaderType::Fragmented;
+		StartRequest->MediaURL = ma->GetMediaURL();
+		StartRequest->AssetID = ma->GetUniqueIdentifier();
+		StartRequest->EarliestPTS = bFrameAccurateSearch ? StartPosition.Time : FTimeValue().SetFromTimespan(SearchTime);
+		StartRequest->LastPTS = PlayRangeEnd;
+		StartRequest->TimestampSequenceIndex = InSequenceState.GetSequenceIndex();
+		for(int32 i=0; i<3; ++i)
+		{
+			if (!Requests[i])
+			{
+				continue;
+			}
+			Requests[i]->MediaAsset = StartRequest->MediaAsset;
+			Requests[i]->LoaderType = StartRequest->LoaderType;
+			Requests[i]->AssetID = StartRequest->AssetID;
+			Requests[i]->MediaURL = StartRequest->MediaURL;
+			Requests[i]->EarliestPTS = StartRequest->EarliestPTS;
+			Requests[i]->LastPTS = StartRequest->LastPTS;
+			Requests[i]->TimestampSequenceIndex = StartRequest->TimestampSequenceIndex;
+
+			if (LoaderType == FTimelineAssetISOBMFF::ELoaderType::Chunks)
+			{
+				// TBD: aggregate consecutive chunks into a single request?
+				AggregateChunkSamples(Requests[i], AggregateMinChunkDuration);
+			}
+
+			// If this isn't already the last segment we need to check if this segment contains
+			// or is past the playback range end.
+			if (!Requests[i]->bIsLastSegment)
+			{
+				FTimeValue LastDTS = Requests[i]->FirstDTS + Requests[i]->SegmentDuration;
+				if (LastDTS > PlayRangeEnd)
+				{
+					Requests[i]->bIsLastSegment = true;
+				}
+			}
+
+			// Set the first request type as the type for the start request.
+			if (StartRequest->StreamType == EStreamType::Unsupported)
+			{
+				StartRequest->StreamType = Requests[i]->StreamType;
+			}
+			StartRequest->Bitrate += Requests[i]->Bitrate;
+			if (!StartRequest->SegmentDuration.IsValid() || StartRequest->SegmentDuration < Requests[i]->SegmentDuration)
+			{
+				StartRequest->SegmentDuration = Requests[i]->SegmentDuration;
+			}
+			StartRequest->DependentStreams.Emplace(MoveTemp(Requests[i]));
+		}
+		OutSegment = StartRequest;
+		return IManifest::FResult(IManifest::FResult::EType::Found);
+	}
+	else if (LoaderType == FTimelineAssetISOBMFF::ELoaderType::Muxed)
+	{
+		TSharedPtr<FStreamSegmentRequestISOBMFF, ESPMode::ThreadSafe> PrimaryReq;
+		// Pick the selected video or audio track as the primary track we need to start with.
+		if (SelectedVideoMetadata && !SelectedVideoMetadata->ID.IsEmpty())
+		{
+			PrimaryReq = SetupTrackRequest(ma->GetAdaptationSetByTypeAndID(EStreamType::Video, SelectedVideoMetadata->ID), EStreamType::Video);
+		}
+		else if (SelectedAudioMetadata && !SelectedAudioMetadata->ID.IsEmpty())
+		{
+			PrimaryReq = SetupTrackRequest(ma->GetAdaptationSetByTypeAndID(EStreamType::Audio, SelectedAudioMetadata->ID), EStreamType::Audio);
+		}
+		if (!PrimaryReq || !PrimaryReq->Track.TrackIterator)
+		{
+			return SetError();
+		}
+		if (!bFrameAccurateSearch)
+		{
+			SearchTime = PrimaryReq->Track.TrackIterator->GetEffectivePTS().GetAsTimespan();
+		}
+
+		TSharedPtr<FStreamSegmentRequestISOBMFF, ESPMode::ThreadSafe> StartRequest = MakeShared<FStreamSegmentRequestISOBMFF, ESPMode::ThreadSafe>();
+		StartRequest->MediaAsset = ma;
+		StartRequest->LoaderType = FStreamSegmentRequestISOBMFF::ELoaderType::Muxed;
+		StartRequest->MediaURL = ma->GetMediaURL();
+		StartRequest->AssetID = ma->GetUniqueIdentifier();
+		StartRequest->EarliestPTS = bFrameAccurateSearch ? StartPosition.Time : FTimeValue().SetFromTimespan(SearchTime);
+		StartRequest->LastPTS = PlayRangeEnd;
+		StartRequest->TimestampSequenceIndex = InSequenceState.GetSequenceIndex();
+		StartRequest->StreamType = PrimaryReq->StreamType;
+		StartRequest->AdaptationSetID = PrimaryReq->AdaptationSetID;
+		StartRequest->RepresentationID = PrimaryReq->RepresentationID;
+		StartRequest->AdaptationSetTrackIndex = PrimaryReq->AdaptationSetTrackIndex;
+		// Using the sample information from the "primary" track we now need to go over all tracks and set them up.
+		EStreamType AllStreamTypes[] { EStreamType::Video, EStreamType::Audio, EStreamType::Subtitle };
+		FString SelectedStreamTypeIDs[3];
+		SelectedStreamTypeIDs[0] = SelectedVideoMetadata && !SelectedVideoMetadata->ID.IsEmpty() ? SelectedVideoMetadata->ID : FString();
+		SelectedStreamTypeIDs[1] = SelectedAudioMetadata && !SelectedAudioMetadata->ID.IsEmpty() ? SelectedAudioMetadata->ID : FString();
+		SelectedStreamTypeIDs[2] = SelectedSubtitleMetadata && !SelectedSubtitleMetadata->ID.IsEmpty() ? SelectedSubtitleMetadata->ID : FString();
+		for(int32 stIdx=0; stIdx<UE_ARRAY_COUNT(AllStreamTypes); ++stIdx)
+		{
+			for(int32 i=0,iMax=ma->GetNumberOfAdaptationSets(AllStreamTypes[stIdx]); i<iMax; ++i)
+			{
+				TSharedPtrTS<FAdaptationSetISOBMFF> AS = StaticCastSharedPtr<FAdaptationSetISOBMFF>(ma->GetAdaptationSetByTypeAndIndex(AllStreamTypes[stIdx], i));
+				auto TrkReq = SetupTrackRequest(AS, AllStreamTypes[stIdx]);
+				if (TrkReq)
+				{
+					TrkReq->MediaAsset = StartRequest->MediaAsset;
+					TrkReq->LoaderType = StartRequest->LoaderType;
+					TrkReq->AssetID = StartRequest->AssetID;
+					TrkReq->MediaURL = StartRequest->MediaURL;
+					TrkReq->EarliestPTS = StartRequest->EarliestPTS;
+					TrkReq->LastPTS = StartRequest->LastPTS;
+					TrkReq->TimestampSequenceIndex = StartRequest->TimestampSequenceIndex;
+					TrkReq->bWasSelectedAtStart = AS->GetUniqueIdentifier().Equals(SelectedStreamTypeIDs[stIdx]);
+					// Add this track request to the list of dependent streams.
+					// Note: This track could be at EOS at that time already, which is indicated by it not having an iterator!
+					StartRequest->DependentStreams.Emplace(MoveTemp(TrkReq));
+				}
+			}
+		}
+		AdjustMuxedDownload(StartRequest, MuxedSegmentDuration, EAdjustMuxType::FirstStart);
+		OutSegment = StartRequest;
+		return IManifest::FResult(IManifest::FResult::EType::Found);
+	}
+	else
+	{
+		return IManifest::FResult(IManifest::FResult::EType::NotFound);
+	}
+}
+
+//-----------------------------------------------------------------------------
+/**
+ * Creates the next segment request.
+ */
+IManifest::FResult FManifestISOBMFFInternal::FPlayPeriodISOBMFF::GetNextSegment(TSharedPtrTS<IStreamSegment>& OutSegment, TSharedPtrTS<const IStreamSegment> CurrentSegment, const FPlayStartOptions& Options)
+{
+	const FStreamSegmentRequestISOBMFF* Request = static_cast<const FStreamSegmentRequestISOBMFF*>(CurrentSegment.Get());
+	if (Request)
+	{
+		TSharedPtrTS<FTimelineAssetISOBMFF> ma = MediaAsset.Pin();
+		if (!ma)
+		{
+			return IManifest::FResult(IManifest::FResult::EType::NotFound);
+		}
+		// Check if the current request did not already go up to the end of the stream. If so there is no next segment.
+		if (!Request->bIsLastSegment)
+		{
+			if (LoaderType == FTimelineAssetISOBMFF::ELoaderType::Chunks || LoaderType == FTimelineAssetISOBMFF::ELoaderType::Fragmented)
+			{
+				TSharedPtr<FStreamSegmentRequestISOBMFF, ESPMode::ThreadSafe> NextRequest = MakeShared<FStreamSegmentRequestISOBMFF, ESPMode::ThreadSafe>();
+				*NextRequest = *Request;
+				NextRequest->DownloadStats = {};
+				NextRequest->ConnectionInfo = {};
+
+				if (NextRequest->Track.TrackIterator)
+				{
+					NextRequest->Track.TrackIterator->GetCurrentChunkRemainingSampleInfo(NextRequest->Track.TrackChunkSampleInfo);
+					NextRequest->SegmentDuration.SetFromND(NextRequest->Track.TrackChunkSampleInfo.Duration, NextRequest->Track.TrackIterator->GetTimescale());
+					auto DTS = NextRequest->Track.TrackIterator->GetEffectiveDTS();
+					auto PTS = NextRequest->Track.TrackIterator->GetEffectivePTS();
+					NextRequest->FirstDTS.SetFromND(DTS.GetNumerator(), DTS.GetDenominator());
+					NextRequest->bIsLastSegment = NextRequest->Track.TrackChunkSampleInfo.bIsLastChunk;
+				}
+				else
+				{
+					NextRequest->bIsLastSegment = true;
+				}
+
+				// Check if this is still covered by the playback range
+				FTimeValue PlayRangeEnd = Options.PlaybackRange.End;
+				if (PlayRangeEnd > ma->GetDuration())
+				{
+					PlayRangeEnd = ma->GetDuration();
+				}
+				NextRequest->LastPTS = PlayRangeEnd;
+				if (NextRequest->FirstDTS > PlayRangeEnd)
+				{
+					return IManifest::FResult(IManifest::FResult::EType::PastEOS);
+				}
+
+				OutSegment = MoveTemp(NextRequest);
+				return IManifest::FResult(IManifest::FResult::EType::Found);
+			}
+			else if (LoaderType == FTimelineAssetISOBMFF::ELoaderType::Muxed)
+			{
+				// The next request for a multiplex merely continues where this one ended.
+				TSharedPtr<FStreamSegmentRequestISOBMFF, ESPMode::ThreadSafe> NextRequest = MakeShared<FStreamSegmentRequestISOBMFF, ESPMode::ThreadSafe>();
+				*NextRequest = *Request;
+				NextRequest->DownloadStats = {};
+				NextRequest->ConnectionInfo = {};
+				AdjustMuxedDownload(NextRequest, MuxedSegmentDuration, EAdjustMuxType::Continuation);
+				if (NextRequest->bIsLastSegment)
+				{
+					return IManifest::FResult(IManifest::FResult::EType::PastEOS);
+				}
+				OutSegment = MoveTemp(NextRequest);
+				return IManifest::FResult(IManifest::FResult::EType::Found);
+			}
+			else
+			{
+				unimplemented();
+				return IManifest::FResult(IManifest::FResult::EType::NotFound);
+			}
+		}
+	}
+	return IManifest::FResult(IManifest::FResult::EType::PastEOS);
+}
+
+
+//-----------------------------------------------------------------------------
+/**
+ * Sets up a starting segment request to loop playback to.
+ * The streams selected through SelectStream() will be used.
+ */
+IManifest::FResult FManifestISOBMFFInternal::FPlayPeriodISOBMFF::GetLoopingSegment(TSharedPtrTS<IStreamSegment>& OutSegment, const FPlayerSequenceState& SequenceState, const FPlayStartPosition& StartPosition, ESearchType SearchType)
+{
+	return GetStartingSegment(OutSegment, SequenceState, StartPosition, SearchType);
+}
+
+
+//-----------------------------------------------------------------------------
+/**
+ * Same as GetStartingSegment() except this is for a specific stream (video, audio, ...) only.
+ * To be used when a track (language) change is made and a new segment is needed at the current playback position.
+ */
+IManifest::FResult FManifestISOBMFFInternal::FPlayPeriodISOBMFF::GetContinuationSegment(TSharedPtrTS<IStreamSegment>& OutSegment, EStreamType InStreamType, const FPlayerSequenceState& InSequenceState, const FPlayStartPosition& InStartPosition, ESearchType InSearchType)
+{
+	if (LoaderType == FTimelineAssetISOBMFF::ELoaderType::Chunks || LoaderType == FTimelineAssetISOBMFF::ELoaderType::Fragmented)
+	{
+		TSharedPtrTS<IStreamSegment> NewReq;
+		FResult Result = GetStartingSegment(NewReq, InSequenceState, InStartPosition, InSearchType);
+		if (Result.IsSuccess())
+		{
+			// We do not need the starting request itself, only the dependent stream of the desired type.
+			const FStreamSegmentRequestISOBMFF* Request = static_cast<const FStreamSegmentRequestISOBMFF*>(NewReq.Get());
+			for(int32 i=0; i<Request->DependentStreams.Num(); ++i)
+			{
+				if (Request->DependentStreams[i]->StreamType == InStreamType)
+				{
+					OutSegment = Request->DependentStreams[i];
+					return IManifest::FResult(IManifest::FResult::EType::Found);
+				}
+			}
+		}
+	}
+	return IManifest::FResult(IManifest::FResult::EType::NotFound);
+}
+
+
+//-----------------------------------------------------------------------------
+/**
+ * Creates a segment retry request.
+ */
+IManifest::FResult FManifestISOBMFFInternal::FPlayPeriodISOBMFF::GetRetrySegment(TSharedPtrTS<IStreamSegment>& OutSegment, TSharedPtrTS<const IStreamSegment> CurrentSegment, const FPlayStartOptions& Options, bool bReplaceWithFillerData)
+{
+	const FStreamSegmentRequestISOBMFF* Request = static_cast<const FStreamSegmentRequestISOBMFF*>(CurrentSegment.Get());
+	if (Request)
+	{
+		if (LoaderType == FTimelineAssetISOBMFF::ELoaderType::Chunks || LoaderType == FTimelineAssetISOBMFF::ELoaderType::Fragmented)
+		{
+			// The retry is actually a copy of the current request with just updated
+			// remaining sample per chunk values of what has not been read so far.
+			TSharedPtr<FStreamSegmentRequestISOBMFF, ESPMode::ThreadSafe> NextRequest = MakeShared<FStreamSegmentRequestISOBMFF, ESPMode::ThreadSafe>();
+			*NextRequest = *Request;
+			NextRequest->DownloadStats = {};
+			NextRequest->ConnectionInfo = {};
+			++NextRequest->NumOverallRetries;
+			if (NextRequest->Track.TrackIterator)
+			{
+				NextRequest->Track.TrackIterator->GetCurrentChunkRemainingSampleInfo(NextRequest->Track.TrackChunkSampleInfo);
+				NextRequest->SegmentDuration.SetFromND(NextRequest->Track.TrackChunkSampleInfo.Duration, NextRequest->Track.TrackIterator->GetTimescale());
+				auto DTS = NextRequest->Track.TrackIterator->GetEffectiveDTS();
+				auto PTS = NextRequest->Track.TrackIterator->GetEffectivePTS();
+				NextRequest->FirstDTS.SetFromND(DTS.GetNumerator(), DTS.GetDenominator());
+				NextRequest->bIsLastSegment = NextRequest->Track.TrackChunkSampleInfo.bIsLastChunk;
+			}
+			else
+			{
+				// This cannot happen because there will be an iterator. Just in case it went away though...
+				return IManifest::FResult(IManifest::FResult::EType::NotFound);
+			}
+
+			OutSegment = MoveTemp(NextRequest);
+			return IManifest::FResult(IManifest::FResult::EType::Found);
+		}
+		else if (LoaderType == FTimelineAssetISOBMFF::ELoaderType::Muxed)
+		{
+			// The retry request for a multiplex merely continues where this one failed.
+			TSharedPtr<FStreamSegmentRequestISOBMFF, ESPMode::ThreadSafe> NextRequest = MakeShared<FStreamSegmentRequestISOBMFF, ESPMode::ThreadSafe>();
+			*NextRequest = *Request;
+			NextRequest->DownloadStats = {};
+			NextRequest->ConnectionInfo = {};
+			++NextRequest->NumOverallRetries;
+			AdjustMuxedDownload(NextRequest, MuxedSegmentDuration, EAdjustMuxType::Retry);
+			OutSegment = MoveTemp(NextRequest);
+			return IManifest::FResult(IManifest::FResult::EType::Found);
+		}
+		else
+		{
+			unimplemented();
+		}
+	}
+	return IManifest::FResult(IManifest::FResult::EType::NotFound);
+}
+
+
+//-----------------------------------------------------------------------------
+/**
+ * Returns the average segment duration.
+ */
+void FManifestISOBMFFInternal::FPlayPeriodISOBMFF::GetAverageSegmentDuration(FTimeValue& OutAverageSegmentDuration, const FString& AdaptationSetID, const FString& RepresentationID)
+{
+	TSharedPtrTS<FTimelineAssetISOBMFF> ma = MediaAsset.Pin();
+	if (ma.IsValid())
+	{
+		ma->GetAverageSegmentDuration(OutAverageSegmentDuration, AdaptationSetID, RepresentationID);
+	}
+}
+
+
+//-----------------------------------------------------------------------------
+/**
+ * Called by the ABR to increase the delay in fetching the next segment in case the segment returned a 404 when fetched at
+ * the announced availability time. This may reduce 404's on the next segment fetches.
+ */
+void FManifestISOBMFFInternal::FPlayPeriodISOBMFF::IncreaseSegmentFetchDelay(const FTimeValue& IncreaseAmount)
+{
+	// No-op.
+}
+
+
+//-----------------------------------------------------------------------------
+/**
+ * Builds the timeline asset.
+ */
+FErrorDetail FManifestISOBMFFInternal::FTimelineAssetISOBMFF::Build(IPlayerSessionServices* InPlayerSessionServices, TArray<TSharedPtr<MP4Boxes::FMP4BoxBase>>&& InParsedRootBoxes, const FString& InURL,
+																	TSharedPtr<FMP4DataLoader, ESPMode::ThreadSafe> InDataLoader, FMP4DataLoader::FCancellationCheckDelegate InCancelCheck)
+{
+	PlayerSessionServices = InPlayerSessionServices;
+	ParsedRootBoxes = MoveTemp(InParsedRootBoxes);
+	MediaURL = InURL;
+	bIsLocalFile = MediaURL.StartsWith(TEXT("file:"), ESearchCase::IgnoreCase);
+
+	FErrorDetail Error;
+	auto SetError = [&Error](const FString& InMsg) -> void
+	{
+		Error.SetFacility(Facility::EFacility::MP4Playlist);
+		Error.SetMessage(InMsg);
+		Error.SetCode(ERRCODE_MANIFEST_ISOBMFF_NO_PLAYABLE_STREAMS);
+	};
+
+	auto FindRootBox = [&](uint32 InType) -> TSharedPtr<MP4Boxes::FMP4BoxBase>
+	{
+		auto rb = ParsedRootBoxes.FindByPredicate([Type=InType](const TSharedPtr<MP4Boxes::FMP4BoxBase>& In) { return In->GetType() == Type;});
+		return rb ? *rb : nullptr;
+	};
+	TSharedPtr<MP4Boxes::FMP4BoxMOOV, ESPMode::ThreadSafe> MoovBox = StaticCastSharedPtr<MP4Boxes::FMP4BoxMOOV>(FindRootBox(MP4Utilities::MakeBoxAtom('m','o','o','v')));
+	if (!MoovBox)
+	{
+		SetError(TEXT("No `moov` box found, cannot play this file."));
+		return Error;
+	}
+	// Get the `mvhd` box for the duration of the movie and the timescale other values are measured in.
+	auto MvhdBox = MoovBox->FindBoxRecursive<MP4Boxes::FMP4BoxMVHD>(MP4Utilities::MakeBoxAtom('m','v','h','d'), 0);
+	if (!MvhdBox)
+	{
+		SetError(TEXT("No `mvhd` box found. Invalid file."));
+		return Error;
+	}
+
+	MP4Utilities::FFractionalTime FullMovieDuration = MvhdBox->GetDuration();
+
+	// Check if this is a fragmented file. We do not look for `sidx` or `mfra` boxes since these are purely optional.
+	TSharedPtr<MP4Boxes::FMP4BoxMOOF, ESPMode::ThreadSafe> FirstMoofBox = StaticCastSharedPtr<MP4Boxes::FMP4BoxMOOF>(FindRootBox(MP4Utilities::MakeBoxAtom('m','o','o','f')));
+	auto MvexBox = MoovBox->FindBoxRecursive<MP4Boxes::FMP4BoxMVEX>(MP4Utilities::MakeBoxAtom('m','v','e','x'), 0);
+	if (FirstMoofBox && !MvexBox)
+	{
+		SetError(TEXT("No `mvex` box present in a file containing `moof` fragment boxes."));
+		return Error;
+	}
+	const bool bIsFragmented = MvexBox || FirstMoofBox;
+	if (bIsFragmented)
+	{
+		auto MehdBox = MvexBox->FindBoxRecursive<MP4Boxes::FMP4BoxMEHD>(MP4Utilities::MakeBoxAtom('m','e','h','d'), 0);
+		if (MehdBox)
+		{
+			FullMovieDuration.SetNumerator(MehdBox->GetFragmentDuration());
+		}
+		else if (FullMovieDuration.GetNumerator() == 0)
+		{
+			SetError(TEXT("Duration in `mvhd` box is set to zero in a fragmented file that is missing a `mehd` box. This is not a usable mp4/mov file."));
+			return Error;
+		}
+	}
+
+	TArray<TSharedPtr<FTrackInfo, ESPMode::ThreadSafe>> UsableTracks;
+	TArray<TSharedPtr<MP4Boxes::FMP4BoxTRAK, ESPMode::ThreadSafe>> AllTracks;
+	MoovBox->GetAllBoxInstances(AllTracks, MP4Utilities::MakeBoxAtom('t','r','a','k'));
+	int64 LongestTkD = 0;
+	int64 ShortestTkD = TNumericLimits<int64>::Max();
+	TMap<TSharedPtr<FTrackInfo, ESPMode::ThreadSafe>, TSharedPtr<MP4Boxes::FMP4BoxTRAK, ESPMode::ThreadSafe>> TrakMap;
+	TSet<uint32> ReferencedTimecodeTrackIDs;
+	TArray<TSharedPtr<FTrackInfo, ESPMode::ThreadSafe>> ExistingTimecodeTracks;
+	for(int32 i=0; i<AllTracks.Num(); ++i)
+	{
+		const TSharedPtr<MP4Boxes::FMP4BoxTRAK, ESPMode::ThreadSafe>& ThisTrakBox(AllTracks[i]);
+
+		TSharedPtr<FTrackInfo, ESPMode::ThreadSafe> TrkInf = MakeShared<FTrackInfo, ESPMode::ThreadSafe>();
+		auto Tkhd = ThisTrakBox->FindBoxRecursive<MP4Boxes::FMP4BoxTKHD>(MP4Utilities::MakeBoxAtom('t','k','h','d'), 1);
+		auto Mdhd = ThisTrakBox->FindBoxRecursive<MP4Boxes::FMP4BoxMDHD>(MP4Utilities::MakeBoxAtom('m','d','h','d'), 3);
+		if (!Tkhd || !Mdhd)
+		{
+			continue;
+		}
+		TrkInf->TrackID = Tkhd->GetTrackID();
+		TrkInf->TrackTimescale = Mdhd->GetTimescale();
+		int64 TrackDuration = Tkhd->GetDuration();
+		if (TrackDuration == TNumericLimits<int64>::Max())
+		{
+			continue;
+		}
+
+		// Take note if this track references a timecode track. We need to do this for every track, even those that
+		// we cannot use to prevent using any unreferenced timecode track for another usable track if the track
+		// itself does not reference any timecode track.
+		auto Tref = ThisTrakBox->FindBoxRecursive<MP4Boxes::FMP4BoxTREF>(MP4Utilities::MakeBoxAtom('t','r','e','f'), 1);
+		if (Tref)
+		{
+			auto References = Tref->GetEntriesOfType(MP4Utilities::MakeBoxAtom('t','m','c','d'));
+			if (References.Num())
+			{
+				if (References.Num() > 1)
+				{
+					LogMessage(IInfoLog::Warning, FString::Printf(TEXT("Track #%u contains more than one `tmcd` reference box. Using first reference only."), TrkInf->TrackID));
+				}
+				if (References[0].TrackIDs.Num())
+				{
+					if (References[0].TrackIDs.Num() > 1)
+					{
+						LogMessage(IInfoLog::Warning, FString::Printf(TEXT("Track #%u references more than one timecode track. Using first reference only."), TrkInf->TrackID));
+					}
+					TrkInf->ReferencedTimecodeTrackID = References[0].TrackIDs[0];
+					// Tag every track nonetheless.
+					for(auto& RefTkId : References[0].TrackIDs)
+					{
+						ReferencedTimecodeTrackIDs.Add(RefTkId);
+					}
+				}
+			}
+		}
+		/*
+			if (!Tkhd->IsEnabled())
+			{
+				continue;
+			}
+		*/
+
+		FString Msg;
+		if (!GetTrackFormatInfo(Msg, TrkInf->CodecFormat, TrkInf->DrmFormat, TrkInf->DecoderInformation, ThisTrakBox, TrkInf->TrackID, true))
+		{
+			continue;
+		}
+		// Skip tracks that cannot be decoded.
+		if (!TrkInf->DecoderInformation.bIsDecodable)
+		{
+			continue;
+		}
+		// Skip encrypted tracks.
+		if (TrkInf->DrmFormat.bIsEncrypted)
+		{
+			continue;
+		}
+
+		// Can this track be decoded and meets the current playback restrictions?
+		FStreamCodecInformation ci(TrkInf->CodecFormat);
+		IPlayerStreamFilter* StreamFilter = PlayerSessionServices->GetStreamFilter();
+		if (TrkInf->CodecFormat.Type == FCodecTypeFormat::EType::Timecode || (StreamFilter && StreamFilter->CanDecodeStream(ci)))
+		{
+			// If the file is fragmented we collect all the track fragment boxes for this track.
+			if (bIsFragmented)
+			{
+				TrkInf->FragmentInfo = MakeShared<MP4Boxes::FMP4Track::FFragmentInfo, ESPMode::ThreadSafe>();
+				if (!TrkInf->FragmentInfo->Prepare(ThisTrakBox, MvexBox, ParsedRootBoxes))
+				{
+					LogMessage(IInfoLog::Warning, FString::Printf(TEXT("Fragmented track #%u could not be prepared, ignoring this track."), TrkInf->TrackID));
+					continue;
+				}
+
+				// The summed duration is in media timescale, so we need to transform it into movie timescale.
+				FTimeFraction MediaDur(TrkInf->FragmentInfo->GetTotalDuration(), Mdhd->GetTimescale());
+				TrackDuration = MediaDur.GetAsTimebase(FullMovieDuration.GetDenominator());
+				// It is possible that the movie duration is only describing the content contained in the moov box.
+				// If that is smaller than the track duration, update the movie duration with the track duration.
+				if (TrkInf->CodecFormat.Type != FCodecTypeFormat::EType::Timecode && FullMovieDuration.GetNumerator() < TrackDuration)
+				{
+					FullMovieDuration.SetNumerator(TrackDuration);
+				}
+			}
+
+			// Update track durations only for video and audio tracks.
+			if (TrkInf->CodecFormat.Type == FCodecTypeFormat::EType::Video || TrkInf->CodecFormat.Type == FCodecTypeFormat::EType::Audio)
+			{
+				if (TrackDuration > LongestTkD)
+				{
+					LongestTkD = TrackDuration;
+				}
+				if (TrackDuration < ShortestTkD)
+				{
+					ShortestTkD = TrackDuration;
+				}
+			}
+			TrakMap.Add(TrkInf, ThisTrakBox);
+			// If this is a timecode add it to a separate list for convenience later.
+			if (TrkInf->CodecFormat.Type == FCodecTypeFormat::EType::Timecode)
+			{
+				ExistingTimecodeTracks.Emplace(TrkInf);
+			}
+			UsableTracks.Emplace(MoveTemp(TrkInf));
+		}
+	}
+
+	MP4Utilities::FFractionalTime AdjustedMovieDuration;
+	bool bTruncateToShortestTrack = DEFAULT_TRUNCATE_MOVIE_TO_SHORTEST_TRACK;
+	auto UserOpt = PlayerSessionServices->GetOptionValue(ISOBMFF::OptionKeyISOBMFFTruncateToShortestTrack);
+	if (UserOpt.IsValid())
+	{
+		bTruncateToShortestTrack = UserOpt.SafeGetBool(false);
+	}
+	int32 CVarOpt = CVarElectraISOBMFFPlayerForceDuration.GetValueOnAnyThread();
+	if (CVarOpt == 1)
+	{
+		bTruncateToShortestTrack = true;
+	}
+	else if (CVarOpt == 2)
+	{
+		bTruncateToShortestTrack = false;
+	}
+	if (bTruncateToShortestTrack)
+	{
+		OverallTrackDuration.SetFromND(ShortestTkD, FullMovieDuration.GetDenominator());
+		AdjustedMovieDuration.SetFromND(ShortestTkD, FullMovieDuration.GetDenominator());
+	}
+	else
+	{
+		OverallTrackDuration.SetFromND(FullMovieDuration.GetNumerator(), FullMovieDuration.GetDenominator());
+		AdjustedMovieDuration = FullMovieDuration;
+	}
+
+	int32 NumAddedTracks = 0;
+	for(int32 i=0; i<UsableTracks.Num(); ++i)
+	{
+		TSharedPtr<FTrackInfo, ESPMode::ThreadSafe> TrkInf = UsableTracks[i];
+
+		auto TRAKBox = TrakMap[TrkInf];
+		TrkInf->MP4Track = MP4Boxes::FMP4Track::Create(TRAKBox, TrkInf->FragmentInfo);
+		if (!TrkInf->MP4Track->Prepare(FullMovieDuration, AdjustedMovieDuration))
+		{
+			continue;
+		}
+		// Timecode tracks were set up for convenience only. They are not to be handled any further.
+		if (TrkInf->CodecFormat.Type == FCodecTypeFormat::EType::Timecode)
+		{
+			continue;
+		}
+
+		// In an mp4 file we treat every track as a single adaptation set with one representation only.
+		// That's because by definition an adaptation set contains the same content at different bitrates and resolutions, but
+		// the type, language and codec has to be the same.
+		switch(TrkInf->CodecFormat.Type)
+		{
+			case FCodecTypeFormat::EType::Video:
+			{
+				TSharedPtrTS<FAdaptationSetISOBMFF> AdaptationSet = MakeSharedTS<FAdaptationSetISOBMFF>();
+				AdaptationSet->CreateFrom(TrkInf, FStreamCodecInformation(TrkInf->CodecFormat), NumAddedTracks, VideoAdaptationSets.Num());
+				VideoAdaptationSets.Add(MoveTemp(AdaptationSet));
+				break;
+			}
+			case FCodecTypeFormat::EType::Audio:
+			{
+				TSharedPtrTS<FAdaptationSetISOBMFF> AdaptationSet = MakeSharedTS<FAdaptationSetISOBMFF>();
+				AdaptationSet->CreateFrom(TrkInf, FStreamCodecInformation(TrkInf->CodecFormat), NumAddedTracks, AudioAdaptationSets.Num());
+				AudioAdaptationSets.Add(MoveTemp(AdaptationSet));
+				break;
+			}
+			case FCodecTypeFormat::EType::Subtitle:
+			{
+				TSharedPtrTS<FAdaptationSetISOBMFF> AdaptationSet = MakeSharedTS<FAdaptationSetISOBMFF>();
+				AdaptationSet->CreateFrom(TrkInf, FStreamCodecInformation(TrkInf->CodecFormat), NumAddedTracks, SubtitleAdaptationSets.Num());
+				SubtitleAdaptationSets.Add(MoveTemp(AdaptationSet));
+				break;
+			}
+			default:
+			{
+				break;
+			}
+		}
+		++NumAddedTracks;
+	}
+
+	// No playable content?
+	if (VideoAdaptationSets.Num() == 0 && AudioAdaptationSets.Num() == 0)
+	{
+		SetError(TEXT("No playable streams in this mp4"));
+		return Error;
+	}
+
+	// Determine the presumably best way to load the movie
+	if (bIsFragmented)
+	{
+		// If the movie is fragmented we have to load fragments.
+		LoaderType = ELoaderType::Fragmented;
+	}
+	else if (NumAddedTracks == 1)
+	{
+		// If there is only a single track we use the multiplex loader to reduce the
+		// number of individual chunk requests.
+		LoaderType = ELoaderType::Muxed;
+	}
+	else //if (NumAddedTracks == 1)
+	{
+		// When there are multiple tracks we try to determine the interleave duration.
+		// If it is tight enough we can use the multiplex loader but if it is too
+		// large we load by chunks to not risk any one buffer not filling up enough
+		// while reading multiplexed.
+		LoaderType = ELoaderType::Muxed;
+		const int32 kNumProbeSamples = 100;
+		const FTimespan kMaxPermittedInterleaveDuration(FTimespan::FromSeconds(0.75));
+		bool bProbeDone = false;
+		for(int32 i=0; !bProbeDone && i<UsableTracks.Num(); ++i)
+		{
+			TSharedPtr<FTrackInfo, ESPMode::ThreadSafe> TrkInf = UsableTracks[i];
+			// Only look at video or audio tracks.
+			if (TrkInf->CodecFormat.Type == FCodecTypeFormat::EType::Video || TrkInf->CodecFormat.Type == FCodecTypeFormat::EType::Audio)
+			{
+				auto TkIt = TrkInf->MP4Track->CreateIterator();
+				FTimespan Dur = TkIt->GetDurationAsTimespan();
+				int64 Off = TkIt->GetSampleFileOffset();
+				int64 NxtOff = Off + TkIt->GetSampleSize();
+				for(int32 si=1; si<kNumProbeSamples && TkIt->Next(); ++si)
+				{
+					if (TkIt->GetSampleFileOffset() != NxtOff)
+					{
+						break;
+					}
+					Dur += TkIt->GetDurationAsTimespan();
+					if (Dur > kMaxPermittedInterleaveDuration)
+					{
+						LogMessage(IInfoLog::Verbose, FString::Printf(TEXT("Track interleave duration seems large, using chunk read mode.")));
+						LoaderType = ELoaderType::Chunks;
+						bProbeDone = true;
+						break;
+					}
+					NxtOff += TkIt->GetSampleSize();
+				}
+			}
+		}
+	}
+
+	// Should we parse any timecodes?
+	bool bPrepareFirstTimecodes = PlayerSessionServices ? PlayerSessionServices->HaveOptionValue(Electra::OptionKeyParseTimecodeInfo) : false;
+
+	// Is there a metadata box?
+	auto Udta = MoovBox->FindBoxRecursive<MP4Boxes::FMP4BoxUDTA>(MP4Utilities::MakeBoxAtom('u','d','t','a'), 1);
+	if (Udta)
+	{
+		auto Meta = Udta->FindBoxRecursive<MP4Boxes::FMP4BoxMETA>(MP4Utilities::MakeBoxAtom('m','e','t','a'), 1);
+		if (Meta)
+		{
+			auto Hdlr = Meta->GetHandler();
+			if (Hdlr)
+			{
+				uint32 hdlr = Hdlr->GetHandlerType();
+				uint32 res0 = Hdlr->GetReservedValue(0);
+				TArray<MP4Utilities::FMP4BoxInfo> Boxes;
+				for(auto &ch : Meta->GetChildren())
+				{
+					Boxes.Emplace(ch->GetBoxInfo());
+				}
+				MediaMetadata = MakeSharedTS<FMetadataParser>();
+				if (MediaMetadata->Parse(hdlr, res0, Boxes) == FMetadataParser::EResult::Success)
+				{
+					PlayerSessionServices->SendMessageToPlayer(FPlaylistMetadataUpdateMessage::Create(FTimeValue(), MediaMetadata, false));
+				}
+				else
+				{
+					MediaMetadata.Reset();
+				}
+			}
+		}
+		// When we are to handle first timecodes let's see if there is a global timecode.
+		if (bPrepareFirstTimecodes)
+		{
+			auto CTIMBox = Udta->FindBoxRecursive<MP4Boxes::FMP4BoxBase>(MP4Utilities::MakeBoxAtom(0xa9U,'T','I','M'), 0);
+			auto CTSCBox = Udta->FindBoxRecursive<MP4Boxes::FMP4BoxBase>(MP4Utilities::MakeBoxAtom(0xa9U,'T','S','C'), 0);
+			auto CTSZBox = Udta->FindBoxRecursive<MP4Boxes::FMP4BoxBase>(MP4Utilities::MakeBoxAtom(0xa9U,'T','S','Z'), 0);
+			if (CTIMBox.IsValid() && CTSCBox.IsValid() && CTSZBox.IsValid())
+			{
+				auto GetValue = [](const TSharedPtr<MP4Boxes::FMP4BoxBase, ESPMode::ThreadSafe>& InBox) -> FString
+				{
+					MP4Utilities::FMP4AtomReader timReader(InBox->GetBoxData());
+					FString String;
+					uint16 StringLength, UnknownValue;
+					if (timReader.Read(StringLength)&&  timReader.Read(UnknownValue))
+					{
+						if (timReader.ReadString(String, StringLength))
+						{
+							return String;
+						}
+					}
+					return FString();
+				};
+				FTrackInfo::FFirstSampleTimecode tc;
+				FFrameRate fr;
+				tc.Timecode = GetValue(CTIMBox);
+				tc.Framerate = FString::Printf(TEXT("%s/%s"), *GetValue(CTSCBox), *GetValue(CTSZBox));
+				TOptional<FTimecode> ptc = FTimecode::ParseTimecode(tc.Timecode);
+				if (ptc.IsSet() && TryParseString(fr, *tc.Framerate))
+				{
+					FFrameNumber fn = ptc.GetValue().ToFrameNumber(fr);
+					tc.TimecodeValue = (uint32) fn.Value;
+				}
+				// Note: this also sets it on timecode tracks, which is ok since the value isn't used there.
+				for(int32 i=0; i<UsableTracks.Num(); ++i)
+				{
+					UsableTracks[i]->FirstSampleTimecode = tc;
+				}
+			}
+		}
+	}
+
+	// Resolve explicitly referenced timecode tracks?
+	if (bPrepareFirstTimecodes && ExistingTimecodeTracks.Num())
+	{
+		TMap<uint32, FManifestISOBMFFInternal::FTrackInfo::FFirstSampleTimecode> TrackFirstTimecodeMap;
+		auto LoadFirstTimecodeSample = [&](uint32 InTrackID) -> bool
+		{
+			if (auto TrInfp = ExistingTimecodeTracks.FindByPredicate([InTrackID](const TSharedPtr<FTrackInfo, ESPMode::ThreadSafe>& tk) { return tk->TrackID == InTrackID; } ))
+			{
+				if (auto TkIt = (*TrInfp)->MP4Track ? (*TrInfp)->MP4Track->CreateIterator() : nullptr)
+				{
+					TArray<uint8> Data;
+					Data.SetNumUninitialized(TkIt->GetSampleSize());
+					int64 Result = InDataLoader->ReadData(Data.GetData(), TkIt->GetSampleSize(), TkIt->GetSampleFileOffset(), InCancelCheck);
+					if (Result == TkIt->GetSampleSize())
+					{
+						// Only handle 4 byte size timecode for now.
+						if (TkIt->GetSampleSize() == 4)
+						{
+							FCodecTypeFormat::FTMCDTimecode TimecodeInfo((*TrInfp)->CodecFormat.Properties.Get<FCodecTypeFormat::FTMCDTimecode>());
+							FManifestISOBMFFInternal::FTrackInfo::FFirstSampleTimecode firstTC;
+							firstTC.TimecodeValue = MP4Utilities::GetFromBigEndian(*reinterpret_cast<const uint32*>(Data.GetData()));
+							firstTC.Framerate = TimecodeInfo.GetFrameRate().ToPrettyText().ToString();
+							firstTC.Timecode = TimecodeInfo.ConvertToTimecode(firstTC.TimecodeValue).ToString();
+							TrackFirstTimecodeMap.Add(InTrackID, firstTC);
+							return true;
+						}
+					}
+				}
+			}
+			return false;
+		};
+
+		// If any track explicitly references a timecode track we only handle those.
+		if (ReferencedTimecodeTrackIDs.Num())
+		{
+			// Load the first timecode from each referenced track and populate a map with it.
+			// If this fails we just ignore it.
+			for(auto& rttid : ReferencedTimecodeTrackIDs)
+			{
+				LoadFirstTimecodeSample(rttid);
+			}
+			// Go over the tracks and assign the timecodes we loaded successfully.
+			for(auto &Tki : UsableTracks)
+			{
+				if (Tki->ReferencedTimecodeTrackID.IsSet() && TrackFirstTimecodeMap.Contains(Tki->ReferencedTimecodeTrackID.GetValue()))
+				{
+					Tki->FirstSampleTimecode = TrackFirstTimecodeMap[Tki->ReferencedTimecodeTrackID.GetValue()];
+				}
+			}
+		}
+		// Otherwise, if no track references a timecode explicitly, then if there is only a single timecode track it applies to all tracks.
+		else if (ExistingTimecodeTracks.Num() == 1)
+		{
+			if (LoadFirstTimecodeSample(ExistingTimecodeTracks[0]->TrackID))
+			{
+				for(auto &Tki : UsableTracks)
+				{
+					Tki->FirstSampleTimecode = TrackFirstTimecodeMap[ExistingTimecodeTracks[0]->TrackID];
+				}
+			}
+		}
+	}
+	// If timecodes are to be extracted, set the one of the first video track to begin with.
+	if (bPrepareFirstTimecodes)
+	{
+		if (VideoAdaptationSets.Num())
+		{
+			auto &tc = VideoAdaptationSets[0]->GetRepresentation()->GetTrackInfo()->FirstSampleTimecode;
+			if (tc.IsSet())
+			{
+				const FVariantValue StartTimecodeValue(tc.GetValue().Timecode);
+				const FVariantValue StartTimecodeFrameRate(tc.GetValue().Framerate);
+				PlayerSessionServices->GetMediaInfoDictionary().Set(MediaInfoStartTimecodeValue, StartTimecodeValue);
+				PlayerSessionServices->GetMediaInfoDictionary().Set(MediaInfoStartTimecodeFrameRate, StartTimecodeFrameRate);
+			}
+		}
+	}
+	return Error;
+}
+
+
+void FManifestISOBMFFInternal::FTimelineAssetISOBMFF::LogMessage(IInfoLog::ELevel Level, const FString& Message)
+{
+	if (PlayerSessionServices)
+	{
+		PlayerSessionServices->PostLog(Facility::EFacility::MP4Playlist, Level, Message);
+	}
+}
+
+
+void FManifestISOBMFFInternal::FTimelineAssetISOBMFF::GetAverageSegmentDuration(FTimeValue& OutAverageSegmentDuration, const FString& AdaptationSetID, const FString& RepresentationID)
+{
+	// This is not expected to be called. And if it does we return a dummy entry.
+	OutAverageSegmentDuration.SetFromSeconds(5.0);
+}
+
+
+void FManifestISOBMFFInternal::FAdaptationSetISOBMFF::CreateFrom(TSharedPtr<FTrackInfo, ESPMode::ThreadSafe> InTrack, const FStreamCodecInformation& InCodecInfo, int32 /*InSequentialIndex*/, int32 InIndexOfType)
+{
+	Representation = MakeSharedTS<FRepresentationISOBMFF>();
+	Representation->CreateFrom(InTrack, InCodecInfo);
+	CodecRFC6381 = InTrack->CodecFormat.RFC6381;
+	KindOfTrack = InIndexOfType == 0 ? TEXT("main") : TEXT("translation");
+	LanguageTag = InTrack->CodecFormat.LanguageTag;
+	UniqueIdentifier = Representation->GetUniqueIdentifier();
+	Duration = InTrack->MP4Track->GetMappedTrackDuration().GetAsTimespan();
+	InternalTrackIndex = InIndexOfType;
+}
+
+void FManifestISOBMFFInternal::FRepresentationISOBMFF::CreateFrom(TSharedPtr<FTrackInfo, ESPMode::ThreadSafe> InTrack, const FStreamCodecInformation& InCodecInfo)
+{
+	Track = MoveTemp(InTrack);
+	CodecInformation = InCodecInfo;
+
+	// The unique identifier will be the track ID inside the mp4.
+	// NOTE: This *MUST* be just a number since it gets parsed back out from a string into a number later! Do *NOT* prepend/append any string literals!!
+	UniqueIdentifier = LexToString(Track->TrackID);
+
+	Name = Track->MP4Track->GetCommonMetadata().Name;
+	if (Name.IsEmpty())
+	{
+		Name = FString::Printf(TEXT("%s (ID=%u)"), *Track->MP4Track->GetCommonMetadata().HandlerName, Track->TrackID);
+	}
+
+	// Get bitrate from the average or max bitrate as stored in the track. If not stored it will be 0.
+	Bitrate = Track->CodecFormat.Bitrate ? Track->CodecFormat.Bitrate : Track->CodecFormat.AverageBitrate;
+
+	// With no bitrate available we set some defaults. This is mainly to avoid a bitrate of 0 from being surfaced that would prevent
+	// events like the initial bitrate change that needs to transition away from 0 to something real.
+	if (Bitrate == 0)
+	{
+		switch(Track->CodecFormat.Type)
+		{
+			case FCodecTypeFormat::EType::Video:
+			{
+				Bitrate = 1 * 1024 * 1024;
+				break;
+			}
+			case FCodecTypeFormat::EType::Audio:
+			{
+				Bitrate = 64 * 1024;
+				break;
+			}
+			case FCodecTypeFormat::EType::Subtitle:
+			{
+				Bitrate = 8 * 1024;
+				break;
+			}
+			default:
+			{
+				// Whatever it is, assume it's a low bitrate.
+				Bitrate = 32 * 1024;
+				break;
+			}
+		}
+	}
+	if (!CodecInformation.GetBitrate())
+	{
+		CodecInformation.SetBitrate(Bitrate);
+	}
+}
+
+} // namespace Electra
+
diff --git a/Engine/Plugins/Media/ElectraPlayer/Source/ElectraPlayerRuntime/Private/Runtime/Player/isobmff/PlaylistISOBMFF.h b/Engine/Plugins/Media/ElectraPlayer/Source/ElectraPlayerRuntime/Private/Runtime/Player/isobmff/PlaylistISOBMFF.h
new file mode 100644
--- /dev/null
+++ b/Engine/Plugins/Media/ElectraPlayer/Source/ElectraPlayerRuntime/Private/Runtime/Player/isobmff/PlaylistISOBMFF.h
@@ -0,0 +1,440 @@
+// Copyright Epic Games, Inc. All Rights Reserved.
+
+#pragma once
+
+#include "PlayerCore.h"
+#include "Player/Manifest.h"
+#include "Player/PlaybackTimeline.h"
+#include "Player/PlayerStreamReader.h"
+#include "Utilities/URLParser.h"
+#include "Utilities/MP4Helpers.h"
+#include "MediaStreamMetadata.h"
+#include "MP4Utilities.h"
+#include "MP4Boxes.h"
+#include "MP4Track.h"
+#include "ElectraCDMClient.h"
+
+namespace Electra
+{
+
+/**
+ * This class represents the internal "playlist" of an mp4 file.
+ * All metadata of the tracks it is composed of are maintained by this class.
+ */
+class FManifestISOBMFFInternal : public IManifest, public TSharedFromThis<FManifestISOBMFFInternal, ESPMode::ThreadSafe>
+{
+public:
+	struct FTrackInfo
+	{
+		TSharedPtr<MP4Boxes::FMP4Track::FFragmentInfo, ESPMode::ThreadSafe> FragmentInfo;
+		TSharedPtr<MP4Boxes::FMP4Track, ESPMode::ThreadSafe> MP4Track;
+
+		FCodecTypeFormat CodecFormat;
+		FDRMTypeFormat DrmFormat;
+		FDecoderInformation DecoderInformation;
+		uint32 TrackID = 0;
+		uint32 TrackTimescale = 0;
+
+		struct FFirstSampleTimecode
+		{
+			FString Timecode;
+			FString Framerate;
+			uint32 TimecodeValue = 0;
+		};
+		TOptional<FFirstSampleTimecode> FirstSampleTimecode;
+		TOptional<uint32> ReferencedTimecodeTrackID;
+	};
+
+public:
+	FManifestISOBMFFInternal(IPlayerSessionServices* InPlayerSessionServices);
+	virtual ~FManifestISOBMFFInternal();
+
+	FErrorDetail Build(TArray<TSharedPtr<MP4Boxes::FMP4BoxBase>>&& InParsedRootBoxes, const FString& URL,
+					   TSharedPtr<FMP4DataLoader, ESPMode::ThreadSafe> InDataLoader, FMP4DataLoader::FCancellationCheckDelegate InCancelCheck);
+
+	EMediaFormatType GetMediaFormatType() const override
+	{ return EMediaFormatType::ISOBMFF; }
+	bool ReloadsPlaylistPeriodically() const override
+	{ return false; }
+	bool SupportsMultipleTracksPerKind() const override
+	{ return true; }
+	bool OrderOfTracksCanChangeDynamically() const override
+	{ return false; }
+	EType GetPresentationType() const override
+	{ return EType::OnDemand; }
+	EReplayEventType GetReplayType() const override
+	{ return IManifest::EReplayEventType::NoReplay; }
+	TSharedPtrTS<const FLowLatencyDescriptor> GetLowLatencyDescriptor() const override
+	{ return nullptr; }
+	FTimeValue CalculateCurrentLiveLatency(const FTimeValue& InCurrentPlaybackPosition, const FTimeValue& InEncoderLatency, bool bViaLatencyElement) const override
+	{ return FTimeValue(); }
+	FTimeValue GetAnchorTime() const override
+	{ return FTimeValue::GetZero(); }
+	FTimeRange GetTotalTimeRange() const override
+	{ return MediaAsset.IsValid() ? MediaAsset->GetTimeRange() : FTimeRange(); }
+	FTimeRange GetSeekableTimeRange() const override
+	{ return GetTotalTimeRange(); }
+	FTimeRange GetPlaybackRange(EPlaybackRangeType InRangeType) const override;
+	FTimeValue GetDuration() const override
+	{ return MediaAsset.IsValid() ? MediaAsset->GetDuration() : FTimeValue(); }
+	FTimeValue GetDefaultStartTime() const override
+	{ return DefaultStartTime; }
+	void ClearDefaultStartTime() override
+	{ DefaultStartTime.SetToInvalid(); }
+	FTimeValue GetDefaultEndTime() const override
+	{ return DefaultEndTime; }
+	void ClearDefaultEndTime() override
+	{ DefaultEndTime.SetToInvalid(); }
+	void GetTrackMetadata(TArray<FTrackMetadata>& OutMetadata, EStreamType StreamType) const override;
+	void UpdateRunningMetaData(TSharedPtrTS<FMetadataParser> InUpdatedMetaData) override
+	{ }
+	FTimeValue GetMinBufferTime() const override;
+	FTimeValue GetDesiredLiveLatency() const override
+	{ return FTimeValue(); }
+	ELiveEdgePlayMode GetLiveEdgePlayMode() const override
+	{ return IManifest::ELiveEdgePlayMode::Never; }
+	TRangeSet<double> GetPossiblePlaybackRates(EPlayRateType InForType) const override;
+	TSharedPtrTS<IProducerReferenceTimeInfo> GetProducerReferenceTimeInfo(int64 ID) const override;
+	void UpdateDynamicRefetchCounter() override
+	{ }
+	void PrepareForLooping(int32 InNumLoopsToAdd) override
+	{ }
+	void TriggerClockSync(EClockSyncType InClockSyncType) override
+	{ }
+	void TriggerPlaylistRefresh() override
+	{ }
+	void ReachedStableBuffer() override
+	{ }
+
+	IStreamReader* CreateStreamReaderHandler() override;
+	FResult FindPlayPeriod(TSharedPtrTS<IPlayPeriod>& OutPlayPeriod, const FPlayStartPosition& StartPosition, ESearchType SearchType) override;
+	FResult FindNextPlayPeriod(TSharedPtrTS<IPlayPeriod>& OutPlayPeriod, TSharedPtrTS<const IStreamSegment> CurrentSegment) override;
+
+
+	class FRepresentationISOBMFF : public IPlaybackAssetRepresentation
+	{
+	public:
+		virtual ~FRepresentationISOBMFF() = default;
+
+		void CreateFrom(TSharedPtr<FTrackInfo, ESPMode::ThreadSafe> InTrack, const FStreamCodecInformation& InCodecInfo);
+
+		FString GetUniqueIdentifier() const override
+		{ return UniqueIdentifier; }
+		const FStreamCodecInformation& GetCodecInformation() const override
+		{ return CodecInformation; }
+		int32 GetBitrate() const override
+		{ return Bitrate; }
+		int32 GetQualityIndex() const override
+		{ return 0; }
+		bool CanBePlayed() const override
+		{ return true; }
+
+		const FString& GetName() const
+		{ return Name; }
+		TSharedPtr<FTrackInfo, ESPMode::ThreadSafe> GetTrackInfo() const
+		{ return Track; }
+	private:
+		TSharedPtr<FTrackInfo, ESPMode::ThreadSafe> Track;
+		FStreamCodecInformation CodecInformation;
+		FString UniqueIdentifier;
+		FString Name;
+		int32 Bitrate;
+	};
+
+	class FAdaptationSetISOBMFF : public IPlaybackAssetAdaptationSet
+	{
+	public:
+		virtual ~FAdaptationSetISOBMFF() = default;
+
+		void CreateFrom(TSharedPtr<FTrackInfo, ESPMode::ThreadSafe> InTrack, const FStreamCodecInformation& InCodecInfo, int32 InSequentialIndex, int32 InIndexOfType);
+
+		FString GetUniqueIdentifier() const override
+		{ return UniqueIdentifier; }
+		FString GetListOfCodecs() const override
+		{ return CodecRFC6381; }
+		const BCP47::FLanguageTag& GetLanguageTag() const override
+		{ return LanguageTag; }
+		int32 GetNumberOfRepresentations() const override
+		{ return 1; }
+		bool IsLowLatencyEnabled() const override
+		{ return false; }
+		TSharedPtrTS<IPlaybackAssetRepresentation> GetRepresentationByIndex(int32 RepresentationIndex) const override
+		{ return RepresentationIndex == 0 ? Representation : TSharedPtrTS<IPlaybackAssetRepresentation>(); }
+		TSharedPtrTS<IPlaybackAssetRepresentation> GetRepresentationByUniqueIdentifier(const FString& InUniqueIdentifier) const override
+		{ return Representation.IsValid() && Representation->GetUniqueIdentifier() == InUniqueIdentifier ? Representation : TSharedPtrTS<IPlaybackAssetRepresentation>(); }
+		TSharedPtrTS<FRepresentationISOBMFF> GetRepresentation() const
+		{ return Representation; }
+		FString GetTrackKind() const
+		{ return KindOfTrack; }
+		int32 GetInternalTrackIndex() const
+		{ return InternalTrackIndex; }
+		FTimespan GetInternalTrackDuration() const
+		{ return Duration; }
+	private:
+		TSharedPtrTS<FRepresentationISOBMFF> Representation;
+		BCP47::FLanguageTag LanguageTag;
+		FString CodecRFC6381;
+		FString KindOfTrack;
+		FString UniqueIdentifier;
+		FTimespan Duration;
+		int32 InternalTrackIndex = -1;
+	};
+
+	class FTimelineAssetISOBMFF : public ITimelineMediaAsset, public TSharedFromThis<FTimelineAssetISOBMFF, ESPMode::ThreadSafe>
+	{
+	public:
+		FTimelineAssetISOBMFF()
+			: PlayerSessionServices(nullptr)
+		{ }
+
+		virtual ~FTimelineAssetISOBMFF() = default;
+
+		FErrorDetail Build(IPlayerSessionServices* InPlayerSessionServices, TArray<TSharedPtr<MP4Boxes::FMP4BoxBase>>&& InParsedRootBoxes, const FString& InURL,
+						   TSharedPtr<FMP4DataLoader, ESPMode::ThreadSafe> InDataLoader, FMP4DataLoader::FCancellationCheckDelegate InCancelCheck);
+
+		FTimeRange GetTimeRange() const override
+		{
+			FTimeRange tr;
+			tr.Start.SetToZero();
+			tr.End = GetDuration();
+			return tr;
+		}
+
+		FTimeValue GetDuration() const override
+		{ return OverallTrackDuration; }
+
+		FString GetAssetIdentifier() const override
+		{ return FString("mp4-asset.0"); }
+		FString GetUniqueIdentifier() const override
+		{ return FString("mp4-media.0"); }
+		int32 GetNumberOfAdaptationSets(EStreamType OfStreamType) const override
+		{
+			switch(OfStreamType)
+			{
+				case EStreamType::Video:
+				{
+					return VideoAdaptationSets.Num();
+				}
+				case EStreamType::Audio:
+				{
+					return AudioAdaptationSets.Num();
+				}
+				case EStreamType::Subtitle:
+				{
+					return SubtitleAdaptationSets.Num();
+				}
+				default:
+				{
+					return 0;
+				}
+			}
+		}
+		TSharedPtrTS<IPlaybackAssetAdaptationSet> GetAdaptationSetByTypeAndIndex(EStreamType OfStreamType, int32 AdaptationSetIndex) const override
+		{
+			switch(OfStreamType)
+			{
+				case EStreamType::Video:
+				{
+					return AdaptationSetIndex < VideoAdaptationSets.Num() ? VideoAdaptationSets[AdaptationSetIndex] : TSharedPtrTS<IPlaybackAssetAdaptationSet>();
+				}
+				case EStreamType::Audio:
+				{
+					return AdaptationSetIndex < AudioAdaptationSets.Num() ? AudioAdaptationSets[AdaptationSetIndex] : TSharedPtrTS<IPlaybackAssetAdaptationSet>();
+				}
+				case EStreamType::Subtitle:
+				{
+					return AdaptationSetIndex < SubtitleAdaptationSets.Num() ? SubtitleAdaptationSets[AdaptationSetIndex] : TSharedPtrTS<IPlaybackAssetAdaptationSet>();
+				}
+			}
+			return TSharedPtrTS<IPlaybackAssetAdaptationSet>();
+		}
+
+		void GetMetaData(TArray<FTrackMetadata>& OutMetadata, EStreamType StreamType) const override
+		{
+			for(int32 i=0, iMax=GetNumberOfAdaptationSets(StreamType); i<iMax; ++i)
+			{
+				TSharedPtrTS<FAdaptationSetISOBMFF> AdaptSet = StaticCastSharedPtr<FAdaptationSetISOBMFF>(GetAdaptationSetByTypeAndIndex(StreamType, i));
+				if (AdaptSet.IsValid())
+				{
+					FTrackMetadata tm;
+					tm.ID = AdaptSet->GetUniqueIdentifier();
+					tm.LanguageTagRFC5646 = AdaptSet->GetLanguageTag();
+					tm.Kind = i==0 ? TEXT("main") : TEXT("translation");
+					tm.Duration = AdaptSet->GetInternalTrackDuration();
+
+					for(int32 j=0, jMax=AdaptSet->GetNumberOfRepresentations(); j<jMax; ++j)
+					{
+						TSharedPtrTS<FRepresentationISOBMFF> Repr = StaticCastSharedPtr<FRepresentationISOBMFF>(AdaptSet->GetRepresentationByIndex(j));
+						if (Repr.IsValid())
+						{
+							FStreamMetadata sd;
+							sd.Bandwidth = Repr->GetBitrate();
+							sd.CodecInformation = Repr->GetCodecInformation();
+							sd.ID = Repr->GetUniqueIdentifier();
+							// There is only 1 "stream" per "track" so we can set the highest bitrate and codec info the same as the track.
+							tm.HighestBandwidth = sd.Bandwidth;
+							tm.HighestBandwidthCodec = sd.CodecInformation;
+
+							tm.Label = Repr->GetName();
+
+							tm.StreamDetails.Emplace(MoveTemp(sd));
+						}
+					}
+					OutMetadata.Emplace(MoveTemp(tm));
+				}
+			}
+		}
+
+		void UpdateRunningMetaData(const FString& InKindOfValue, const FVariant& InNewValue) override
+		{ }
+
+		void GetAverageSegmentDuration(FTimeValue& OutAverageSegmentDuration, const FString& AdaptationSetID, const FString& RepresentationID);
+		const FString& GetMediaURL() const
+		{
+			return MediaURL;
+		}
+
+		TSharedPtrTS<FAdaptationSetISOBMFF> GetAdaptationSetByTypeAndID(EStreamType InStreamType, const FString& InID) const
+		{
+			const TArray<TSharedPtrTS<FAdaptationSetISOBMFF>>* AdSet = nullptr;
+			switch(InStreamType)
+			{
+				case EStreamType::Video:
+				{
+					AdSet = &VideoAdaptationSets;
+					break;
+				}
+				case EStreamType::Audio:
+				{
+					AdSet = &AudioAdaptationSets;
+					break;
+				}
+				case EStreamType::Subtitle:
+				{
+					AdSet = &SubtitleAdaptationSets;
+					break;
+				}
+				default:
+				{
+					return nullptr;
+				}
+			}
+			auto Adpt = (*AdSet).FindByPredicate([InID](const TSharedPtrTS<FAdaptationSetISOBMFF>& e) {return e->GetUniqueIdentifier().Equals(InID);});
+			return Adpt ? *Adpt : nullptr;
+		}
+
+		enum class ELoaderType
+		{
+			Chunks,
+			Muxed,
+			Fragmented
+		};
+		ELoaderType GetLoaderType() const
+		{ return LoaderType; }
+
+	private:
+		struct FPlayRangeEndInfo
+		{
+			void Clear()
+			{
+				Time.SetToInvalid();
+				FileOffsetsPerTrackID.Empty();
+				TotalFileEndOffset = -1;
+			}
+			FTimeValue Time;
+			TArray<int64> FileOffsetsPerTrackID;
+			int64 TotalFileEndOffset = -1;
+		};
+
+		void LogMessage(IInfoLog::ELevel Level, const FString& Message);
+		IPlayerSessionServices* PlayerSessionServices = nullptr;
+		FString MediaURL;
+		bool bIsLocalFile { false };
+		ELoaderType LoaderType { ELoaderType::Chunks };
+		FTimeValue OverallTrackDuration;
+		TArray<TSharedPtr<MP4Boxes::FMP4BoxBase>> ParsedRootBoxes;
+		TArray<TSharedPtrTS<FAdaptationSetISOBMFF>> VideoAdaptationSets;
+		TArray<TSharedPtrTS<FAdaptationSetISOBMFF>> AudioAdaptationSets;
+		TArray<TSharedPtrTS<FAdaptationSetISOBMFF>> SubtitleAdaptationSets;
+		FPlayRangeEndInfo PlayRangeEndInfo;
+		// Metadata from the 'meta' box, if any.
+		TSharedPtrTS<FMetadataParser> MediaMetadata;
+		//TSharedPtrTS<ElectraCDM::IMediaCDMClient> DrmClient;
+	};
+
+
+	class FPlayPeriodISOBMFF : public IManifest::IPlayPeriod
+	{
+	public:
+		FPlayPeriodISOBMFF(TSharedPtrTS<FTimelineAssetISOBMFF> InMediaAsset);
+		virtual ~FPlayPeriodISOBMFF();
+		void SetStreamPreferences(EStreamType ForStreamType, const FStreamSelectionAttributes& StreamAttributes) override;
+		EReadyState GetReadyState() override;
+		void Load() override;
+		void PrepareForPlay() override;
+		int64 GetDefaultStartingBitrate() const override;
+		TSharedPtrTS<FBufferSourceInfo> GetSelectedStreamBufferSourceInfo(EStreamType StreamType) override;
+		FString GetSelectedAdaptationSetID(EStreamType StreamType) override;
+		ETrackChangeResult ChangeTrackStreamPreference(EStreamType ForStreamType, const FStreamSelectionAttributes& StreamAttributes) override;
+		TSharedPtrTS<ITimelineMediaAsset> GetMediaAsset() const override;
+		void SelectStream(const FString& AdaptationSetID, const FString& RepresentationID, int32 QualityIndex, int32 MaxQualityIndex) override;
+		void TriggerInitSegmentPreload(const TArray<FInitSegmentPreload>& InitSegmentsToPreload) override;
+		FResult GetStartingSegment(TSharedPtrTS<IStreamSegment>& OutSegment, const FPlayerSequenceState& InSequenceState, const FPlayStartPosition& StartPosition, ESearchType SearchType) override;
+		FResult GetContinuationSegment(TSharedPtrTS<IStreamSegment>& OutSegment, EStreamType InStreamType, const FPlayerSequenceState& InLoopState, const FPlayStartPosition& InStartPosition, ESearchType InSearchType) override;
+		FResult GetNextSegment(TSharedPtrTS<IStreamSegment>& OutSegment, TSharedPtrTS<const IStreamSegment> CurrentSegment, const FPlayStartOptions& Options) override;
+		FResult GetRetrySegment(TSharedPtrTS<IStreamSegment>& OutSegment, TSharedPtrTS<const IStreamSegment> CurrentSegment, const FPlayStartOptions& Options, bool bReplaceWithFillerData) override;
+		FResult GetLoopingSegment(TSharedPtrTS<IStreamSegment>& OutSegment, const FPlayerSequenceState& InSequenceState, const FPlayStartPosition& StartPosition, ESearchType SearchType) override;
+		void IncreaseSegmentFetchDelay(const FTimeValue& IncreaseAmount) override;
+		void GetAverageSegmentDuration(FTimeValue& OutAverageSegmentDuration, const FString& AdaptationSetID, const FString& RepresentationID) override;
+
+	private:
+		void SelectInitialStream(EStreamType StreamType);
+		TSharedPtrTS<FTrackMetadata> SelectMetadataForAttributes(EStreamType StreamType, const FStreamSelectionAttributes& InAttributes);
+		void MakeBufferSourceInfoFromMetadata(EStreamType StreamType, TSharedPtrTS<FBufferSourceInfo>& OutBufferSourceInfo, TSharedPtrTS<FTrackMetadata> InMetadata);
+		void AggregateChunkSamples(TSharedPtr<IStreamSegment, ESPMode::ThreadSafe> InOutRequest, const FTimespan& InMinDuration);
+		enum class EAdjustMuxType
+		{
+			FirstStart,
+			Continuation,
+			Retry
+		};
+		void AdjustMuxedDownload(TSharedPtr<IStreamSegment, ESPMode::ThreadSafe> InOutRequest, const FTimespan& InMinDuration, EAdjustMuxType InType);
+
+		TWeakPtrTS<FTimelineAssetISOBMFF> MediaAsset;
+		FStreamSelectionAttributes VideoPreferences;
+		FStreamSelectionAttributes AudioPreferences;
+		FStreamSelectionAttributes SubtitlePreferences;
+		TSharedPtrTS<FTrackMetadata> SelectedVideoMetadata;
+		TSharedPtrTS<FTrackMetadata> SelectedAudioMetadata;
+		TSharedPtrTS<FTrackMetadata> SelectedSubtitleMetadata;
+		TSharedPtrTS<FBufferSourceInfo> VideoBufferSourceInfo;
+		TSharedPtrTS<FBufferSourceInfo> AudioBufferSourceInfo;
+		TSharedPtrTS<FBufferSourceInfo> SubtitleBufferSourceInfo;
+		EReadyState CurrentReadyState;
+		FTimelineAssetISOBMFF::ELoaderType LoaderType;
+		FTimespan AggregateMinChunkDuration;
+		FTimespan MuxedSegmentDuration;
+	};
+
+	void LogMessage(IInfoLog::ELevel Level, const FString& Message);
+
+	const TArray<FURL_RFC3986::FQueryParam>& GetURLFragmentComponents() const
+	{ return URLFragmentComponents;	}
+
+	void SetURLFragmentComponents(TArray<FURL_RFC3986::FQueryParam> InURLFragmentComponents)
+	{ URLFragmentComponents = MoveTemp(InURLFragmentComponents); }
+
+	const TArray<FURL_RFC3986::FQueryParam>& GetURLFragmentComponents()
+	{ return URLFragmentComponents;	}
+
+
+	IPlayerSessionServices* PlayerSessionServices = nullptr;
+	TSharedPtrTS<FTimelineAssetISOBMFF> MediaAsset;
+	// The URL fragment components
+	TArray<FURL_RFC3986::FQueryParam> URLFragmentComponents;
+	FTimeValue DefaultStartTime;
+	FTimeValue DefaultEndTime;
+};
+
+
+} // namespace Electra
diff --git a/Engine/Plugins/Media/ElectraPlayer/Source/ElectraPlayerRuntime/Private/Runtime/Player/isobmff/PlaylistReaderISOBMFF.cpp b/Engine/Plugins/Media/ElectraPlayer/Source/ElectraPlayerRuntime/Private/Runtime/Player/isobmff/PlaylistReaderISOBMFF.cpp
new file mode 100644
--- /dev/null
+++ b/Engine/Plugins/Media/ElectraPlayer/Source/ElectraPlayerRuntime/Private/Runtime/Player/isobmff/PlaylistReaderISOBMFF.cpp
@@ -0,0 +1,333 @@
+// Copyright Epic Games, Inc. All Rights Reserved.
+
+#include "PlaylistReaderISOBMFF.h"
+#include "Player/Manifest.h"
+#include "Player/PlayerSessionServices.h"
+#include "HTTP/HTTPManager.h"
+#include "HTTP/HTTPResponseCache.h"
+#include "MP4Utilities.h"
+#include "MP4Boxes.h"
+#include "Utilities/MP4Helpers.h"
+#include "ElectraPlayerPrivate.h"
+#include "Player/isobmff/PlaylistISOBMFF.h"
+#include "Stats/Stats.h"
+#include "HAL/LowLevelMemTracker.h"
+
+#define ERRCODE_ISOBMFF_INVALID_FILE	1
+#define ERRCODE_ISOBMFF_DOWNLOAD_ERROR	2
+
+
+DECLARE_CYCLE_STAT(TEXT("FPlaylistReaderISOBMFF_WorkerThread"), STAT_ElectraPlayer_ISOBMFF_PlaylistWorker, STATGROUP_ElectraPlayer);
+
+
+namespace Electra
+{
+
+/**
+ * This class is responsible for downloading the mp4 non-mdat boxes and parsing them.
+ */
+class FPlaylistReaderISOBMFF : public IPlaylistReaderISOBMFF, public FMediaThread
+{
+public:
+	FPlaylistReaderISOBMFF();
+	void Initialize(IPlayerSessionServices* PlayerSessionServices);
+
+	virtual ~FPlaylistReaderISOBMFF();
+
+	EMediaFormatType GetMediaFormatType() const override
+	{ return EMediaFormatType::ISOBMFF; }
+	void Close() override;
+	void HandleOnce() override
+	{ }
+
+	const FString& GetPlaylistType() const override
+	{
+		static FString Type("mp4");
+		return Type;
+	}
+
+	void LoadAndParse(const FString& InURL, const Playlist::FReplayEventParams& InReplayEventParams) override;
+
+	FString GetURL() const override
+	{ return MainPlaylistURL; }
+
+
+	TSharedPtrTS<IManifest> GetManifest() override
+	{ return Manifest; }
+
+
+private:
+
+	void StartWorkerThread();
+	void StopWorkerThread();
+	void WorkerThread();
+
+	void PostError(const FString& Message, uint16 Code, UEMediaError Error = UEMEDIA_ERROR_OK);
+	void LogMessage(IInfoLog::ELevel Level, const FString& Message);
+
+	IPlayerSessionServices* PlayerSessionServices = nullptr;
+	FString MainPlaylistURL;
+	FString URLFragment;
+	FMediaEvent WorkerThreadQuitSignal;
+	bool bIsWorkerThreadStarted = false;
+
+	TSharedPtr<FMP4DataLoader, ESPMode::ThreadSafe> DataLoader;
+	HTTP::FConnectionInfo ConnectionInfo;
+
+	volatile bool bAbort = false;
+
+	TSharedPtrTS<FManifestISOBMFFInternal> Manifest;
+	FErrorDetail LastErrorDetail;
+};
+
+
+/***************************************************************************************************************************************************/
+/***************************************************************************************************************************************************/
+/***************************************************************************************************************************************************/
+
+TSharedPtrTS<IPlaylistReader> IPlaylistReaderISOBMFF::Create(IPlayerSessionServices* PlayerSessionServices)
+{
+	TSharedPtrTS<FPlaylistReaderISOBMFF> PlaylistReader = MakeSharedTS<FPlaylistReaderISOBMFF>();
+	if (PlaylistReader)
+	{
+		PlaylistReader->Initialize(PlayerSessionServices);
+	}
+	return PlaylistReader;
+}
+
+/***************************************************************************************************************************************************/
+/***************************************************************************************************************************************************/
+/***************************************************************************************************************************************************/
+
+
+FPlaylistReaderISOBMFF::FPlaylistReaderISOBMFF()
+	: FMediaThread("ElectraPlayer::ISOBMFF Playlist")
+{
+}
+
+FPlaylistReaderISOBMFF::~FPlaylistReaderISOBMFF()
+{
+	Close();
+}
+
+void FPlaylistReaderISOBMFF::Initialize(IPlayerSessionServices* InPlayerSessionServices)
+{
+	PlayerSessionServices = InPlayerSessionServices;
+}
+
+void FPlaylistReaderISOBMFF::Close()
+{
+	bAbort = true;
+	StopWorkerThread();
+}
+
+void FPlaylistReaderISOBMFF::StartWorkerThread()
+{
+	check(!bIsWorkerThreadStarted);
+	ThreadStart(FMediaRunnable::FStartDelegate::CreateRaw(this, &FPlaylistReaderISOBMFF::WorkerThread));
+	bIsWorkerThreadStarted = true;
+}
+
+void FPlaylistReaderISOBMFF::StopWorkerThread()
+{
+	if (bIsWorkerThreadStarted)
+	{
+		WorkerThreadQuitSignal.Signal();
+		ThreadWaitDone();
+		ThreadReset();
+		bIsWorkerThreadStarted = false;
+	}
+}
+
+void FPlaylistReaderISOBMFF::PostError(const FString& InMessage, uint16 InCode, UEMediaError InError)
+{
+	LastErrorDetail.Clear();
+	LastErrorDetail.SetError(InError != UEMEDIA_ERROR_OK ? InError : UEMEDIA_ERROR_DETAIL);
+	LastErrorDetail.SetFacility(Facility::EFacility::ISOBMFFPlaylistReader);
+	LastErrorDetail.SetCode(InCode);
+	LastErrorDetail.SetMessage(InMessage);
+	check(PlayerSessionServices);
+	if (PlayerSessionServices)
+	{
+		PlayerSessionServices->PostError(LastErrorDetail);
+	}
+}
+
+void FPlaylistReaderISOBMFF::LogMessage(IInfoLog::ELevel Level, const FString& Message)
+{
+	if (PlayerSessionServices)
+	{
+		PlayerSessionServices->PostLog(Facility::EFacility::ISOBMFFPlaylistReader, Level, Message);
+	}
+}
+
+void FPlaylistReaderISOBMFF::LoadAndParse(const FString& InURL, const Playlist::FReplayEventParams& InReplayEventParams)
+{
+	FURL_RFC3986 UrlParser;
+	UrlParser.Parse(InURL);
+	MainPlaylistURL = UrlParser.Get(true, false);
+	URLFragment = UrlParser.GetFragment();
+
+	StartWorkerThread();
+}
+
+
+void FPlaylistReaderISOBMFF::WorkerThread()
+{
+	LLM_SCOPE(ELLMTag::ElectraPlayer);
+
+	SCOPE_CYCLE_COUNTER(STAT_ElectraPlayer_ISOBMFF_PlaylistWorker);
+	CSV_SCOPED_TIMING_STAT(ElectraPlayer, ISOBMFF_PlaylistWorker);
+
+	TSharedPtrTS<IHTTPResponseCache> LocalCache = IHTTPResponseCache::Create(1024*1024, 32, nullptr);
+
+	DataLoader = FMP4DataLoader::Create();
+	auto DataLoaderCancel = FMP4DataLoader::FCancellationCheckDelegate::CreateLambda([&](){return bAbort;});
+
+	FMP4DataLoader::FHttpParams HttpParams;
+	HttpParams.HTTPManager = PlayerSessionServices->GetHTTPManager();
+	HttpParams.HTTPResponseCache = LocalCache;
+	DataLoader->Open(MainPlaylistURL, HttpParams);
+	bool bIsLocalFile = MainPlaylistURL.StartsWith(TEXT("file:"), ESearchCase::IgnoreCase);
+
+	MP4Utilities::FMP4BoxLocator BoxLocator;
+	TArray<TSharedPtr<MP4Utilities::FMP4BoxData, ESPMode::ThreadSafe>> RootBoxes;
+	const TArray<uint32> FirstBoxes { MP4Utilities::MakeBoxAtom('f','t','y','p'), MP4Utilities::MakeBoxAtom('s','t','y','p'), MP4Utilities::MakeBoxAtom('s','i','d','x'), MP4Utilities::MakeBoxAtom('f','r','e','e'), MP4Utilities::MakeBoxAtom('s','k','i','p') };
+	const TArray<uint32> StopAfterBoxes { MP4Utilities::MakeBoxAtom('m','o','o','v') };	// Empty means to read all boxes and not stop after a specific one.
+	const TArray<uint32> ReadBoxes;			// Empty means to read in all boxes except `mdat`
+	bool bGotBoxes = BoxLocator.LocateAndReadRootBoxes(RootBoxes, DataLoader, FirstBoxes, StopAfterBoxes, ReadBoxes, DataLoaderCancel);
+	ConnectionInfo = DataLoader->GetConnectionInfo();
+	bool bHasErrored = ConnectionInfo.StatusInfo.ErrorDetail.IsSet();
+
+	if (BoxLocator.GetLastError().Len())
+	{
+		PostError(BoxLocator.GetLastError(), ERRCODE_ISOBMFF_INVALID_FILE, UEMEDIA_ERROR_FORMAT_ERROR);
+	}
+
+	// There are currently no retries, so this is the first (and only) attempt we make.
+	const int32 Attempt = 1;
+	// Notify the download of the "main playlist". This indicates the download only, not the parsing thereof.
+	PlayerSessionServices->SendMessageToPlayer(IPlaylistReader::PlaylistDownloadMessage::Create(&ConnectionInfo, Playlist::EListType::Main, Playlist::ELoadType::Initial, Attempt));
+	// Notify that the "main playlist" has been parsed, successfully or not.
+	PlayerSessionServices->SendMessageToPlayer(IPlaylistReader::PlaylistLoadedMessage::Create(LastErrorDetail, &ConnectionInfo, Playlist::EListType::Main, Playlist::ELoadType::Initial, Attempt));
+	// Failed to get the boxes but was not aborted?
+	if (!bAbort && (!bGotBoxes || bHasErrored))
+	{
+		// See if there was a download error
+		if (ConnectionInfo.StatusInfo.ErrorDetail.IsError())
+		{
+			PostError(FString::Printf(TEXT("%s while downloading \"%s\""), *ConnectionInfo.StatusInfo.ErrorDetail.GetMessage(), *ConnectionInfo.EffectiveURL), ERRCODE_ISOBMFF_DOWNLOAD_ERROR, UEMEDIA_ERROR_READ_ERROR);
+		}
+	}
+	else if (bGotBoxes && !bAbort && !bHasErrored)
+	{
+		// Do we have the `moov` box?
+		if (RootBoxes.ContainsByPredicate([](const TSharedPtr<MP4Utilities::FMP4BoxData, ESPMode::ThreadSafe>& In){ return In->Type == MP4Utilities::MakeBoxAtom('m','o','o','v'); }))
+		{
+			if (RootBoxes.ContainsByPredicate([](const TSharedPtr<MP4Utilities::FMP4BoxData, ESPMode::ThreadSafe>& In){ return In->Type == MP4Utilities::MakeBoxAtom('m','d','a','t'); }))
+			{
+				LogMessage(IInfoLog::ELevel::Info, FString::Printf(TEXT("The movie at \"%s\" is not fast-startable. Consider moving the 'moov' box in front of the 'mdat' for faster startup times."), *ConnectionInfo.EffectiveURL));
+			}
+
+			// Parse all the root boxes. There are typically no more than 4 to 6.
+			TArray<TSharedPtr<MP4Boxes::FMP4BoxBase>> ParsedRootBoxes;
+			ParsedRootBoxes.SetNum(RootBoxes.Num());
+			bool bIsFragmented = false;
+			for(int32 i=0; i<RootBoxes.Num(); ++i)
+			{
+				MP4Boxes::FMP4BoxTreeParser tp;
+				if (tp.ParseBoxTree(RootBoxes[i]))
+				{
+					ParsedRootBoxes[i] = tp.GetBoxTree();
+					// Check if this is a fragmented file.
+					if (ParsedRootBoxes[i]->GetType() == MP4Utilities::MakeBoxAtom('m','o','o','v'))
+					{
+						auto MoovBox = StaticCastSharedPtr<MP4Boxes::FMP4BoxMOOV>(ParsedRootBoxes[i]);
+						auto MvexBox = MoovBox->FindBoxRecursive<MP4Boxes::FMP4BoxMVEX>(MP4Utilities::MakeBoxAtom('m','v','e','x'), 0);
+						if (MvexBox)
+						{
+							bIsFragmented = true;
+						}
+					}
+				}
+				else
+				{
+					PostError(FString::Printf(TEXT("Failed to parse the file's box structure.")), ERRCODE_ISOBMFF_INVALID_FILE, UEMEDIA_ERROR_FORMAT_ERROR);
+					bHasErrored = true;
+				}
+			}
+
+			if (bIsFragmented)
+			{
+				if (!bIsLocalFile)
+				{
+					LogMessage(IInfoLog::ELevel::Info, FString::Printf(TEXT("The movie at \"%s\" is fragmented and not a local file. Parsing the movie structure will be slow."), *ConnectionInfo.EffectiveURL));
+				}
+
+				// Continue loading root boxes
+				TArray<TSharedPtr<MP4Utilities::FMP4BoxData, ESPMode::ThreadSafe>> FrgRootBoxes;
+				const TArray<uint32> FrgFirstBoxes;
+				const TArray<uint32> FrgStopAfterBoxes;	// Empty means to read all boxes and not stop after a specific one.
+				const TArray<uint32> FrgReadBoxes;		// Empty means to read in all boxes except `mdat`
+				bGotBoxes = BoxLocator.LocateAndReadRootBoxes(FrgRootBoxes, DataLoader, FrgFirstBoxes, FrgStopAfterBoxes, FrgReadBoxes, DataLoaderCancel);
+				ConnectionInfo = DataLoader->GetConnectionInfo();
+				bHasErrored = ConnectionInfo.StatusInfo.ErrorDetail.IsSet();
+				if (BoxLocator.GetLastError().Len())
+				{
+					PostError(BoxLocator.GetLastError(), ERRCODE_ISOBMFF_INVALID_FILE, UEMEDIA_ERROR_FORMAT_ERROR);
+				}
+				if (!bHasErrored)
+				{
+					int32 CurNumRootBoxes = RootBoxes.Num();
+					RootBoxes.Append(MoveTemp(FrgRootBoxes));
+					ParsedRootBoxes.SetNum(RootBoxes.Num());
+					for(int32 i=CurNumRootBoxes; i<RootBoxes.Num(); ++i)
+					{
+						MP4Boxes::FMP4BoxTreeParser tp;
+						if (tp.ParseBoxTree(RootBoxes[i]))
+						{
+							ParsedRootBoxes[i] = tp.GetBoxTree();
+						}
+						else
+						{
+							PostError(FString::Printf(TEXT("Failed to parse the file's box structure.")), ERRCODE_ISOBMFF_INVALID_FILE, UEMEDIA_ERROR_FORMAT_ERROR);
+							bHasErrored = true;
+						}
+					}
+				}
+			}
+
+			if (!bHasErrored)
+			{
+				TArray<FURL_RFC3986::FQueryParam> URLFragmentComponents;
+				FURL_RFC3986::GetQueryParams(URLFragmentComponents, URLFragment, false);
+				Manifest = MakeSharedTS<FManifestISOBMFFInternal>(PlayerSessionServices);
+				Manifest->SetURLFragmentComponents(MoveTemp(URLFragmentComponents));
+				LastErrorDetail = Manifest->Build(MoveTemp(ParsedRootBoxes), MainPlaylistURL, DataLoader, DataLoaderCancel);
+				if (LastErrorDetail.IsSet())
+				{
+					Manifest.Reset();
+				}
+
+				// Let the external registry know that we have no properties with an end-of-properties call.
+				PlayerSessionServices->ValidateMainPlaylistCustomProperty(GetPlaylistType(), MainPlaylistURL, TArray<FElectraHTTPStreamHeader>(), IPlayerSessionServices::FPlaylistProperty());
+				// Notify that the "variant playlists" are ready. There are no variants in an mp4, but this is the trigger that the playlists are all set up and are good to go now.
+				PlayerSessionServices->SendMessageToPlayer(IPlaylistReader::PlaylistLoadedMessage::Create(LastErrorDetail, &ConnectionInfo, Playlist::EListType::Variant, Playlist::ELoadType::Initial, Attempt));
+			}
+		}
+		else
+		{
+			PostError(FString::Printf(TEXT("No `moov` box found in \"%s\". This is not a valid file."), *ConnectionInfo.EffectiveURL), ERRCODE_ISOBMFF_INVALID_FILE, UEMEDIA_ERROR_FORMAT_ERROR);
+		}
+	}
+
+	DataLoader->Close();
+	DataLoader.Reset();
+	// This thread's work is done. We only wait for termination now.
+	WorkerThreadQuitSignal.Wait();
+}
+
+
+
+
+} // namespace Electra
diff --git a/Engine/Plugins/Media/ElectraPlayer/Source/ElectraPlayerRuntime/Private/Runtime/Player/isobmff/PlaylistReaderISOBMFF.h b/Engine/Plugins/Media/ElectraPlayer/Source/ElectraPlayerRuntime/Private/Runtime/Player/isobmff/PlaylistReaderISOBMFF.h
new file mode 100644
--- /dev/null
+++ b/Engine/Plugins/Media/ElectraPlayer/Source/ElectraPlayerRuntime/Private/Runtime/Player/isobmff/PlaylistReaderISOBMFF.h
@@ -0,0 +1,20 @@
+// Copyright Epic Games, Inc. All Rights Reserved.
+
+#pragma once
+
+#include "PlayerCore.h"
+#include "Player/PlaylistReader.h"
+
+namespace Electra
+{
+class IPlayerSessionServices;
+
+class IPlaylistReaderISOBMFF : public IPlaylistReader
+{
+public:
+	static TSharedPtrTS<IPlaylistReader> Create(IPlayerSessionServices* PlayerSessionServices);
+
+	virtual ~IPlaylistReaderISOBMFF() = default;
+};
+
+} // namespace Electra
diff --git a/Engine/Plugins/Media/ElectraPlayer/Source/ElectraPlayerRuntime/Private/Runtime/Player/isobmff/StreamReaderISOBMFF.cpp b/Engine/Plugins/Media/ElectraPlayer/Source/ElectraPlayerRuntime/Private/Runtime/Player/isobmff/StreamReaderISOBMFF.cpp
new file mode 100644
--- /dev/null
+++ b/Engine/Plugins/Media/ElectraPlayer/Source/ElectraPlayerRuntime/Private/Runtime/Player/isobmff/StreamReaderISOBMFF.cpp
@@ -0,0 +1,913 @@
+// Copyright Epic Games, Inc. All Rights Reserved.
+
+#include "StreamReaderISOBMFF.h"
+#include "PlayerCore.h"
+#include "Player/Manifest.h"
+#include "HTTP/HTTPManager.h"
+#include "Demuxer/ParserISO14496-12.h"
+#include "Stats/Stats.h"
+#include "SynchronizedClock.h"
+#include "Utilities/TimeUtilities.h"
+#include "HAL/LowLevelMemTracker.h"
+#include "ElectraPlayerPrivate.h"
+#include "Player/PlayerStreamReader.h"
+#include "Player/AdaptiveStreamingPlayerABR.h"
+#include "Player/isobmff/OptionKeynamesISOBMFF.h"
+#include "Player/PlayerSessionServices.h"
+#include "Player/PlayerStreamFilter.h"
+
+
+DECLARE_CYCLE_STAT(TEXT("FStreamReaderISOBMFF_HandleRequest"), STAT_ElectraPlayer_ISOBMFF_StreamReader, STATGROUP_ElectraPlayer);
+
+
+namespace Electra
+{
+
+void FStreamSegmentRequestISOBMFF::GetDependentStreams(TArray<TSharedPtrTS<IStreamSegment>>& OutDependentStreams) const
+{
+	OutDependentStreams.Empty();
+	for(auto& Stream : DependentStreams)
+	{
+		OutDependentStreams.Emplace(Stream);
+	}
+}
+
+void FStreamSegmentRequestISOBMFF::GetRequestedStreams(TArray<TSharedPtrTS<IStreamSegment>>& OutRequestedStreams)
+{
+	OutRequestedStreams.Empty();
+	if (LoaderType == FStreamSegmentRequestISOBMFF::ELoaderType::Chunks || LoaderType == FStreamSegmentRequestISOBMFF::ELoaderType::Fragmented)
+	{
+		for(auto& Stream : DependentStreams)
+		{
+			OutRequestedStreams.Emplace(Stream);
+		}
+	}
+	else
+	{
+		OutRequestedStreams.Emplace(AsShared());
+	}
+}
+
+void FStreamSegmentRequestISOBMFF::GetEndedStreams(TArray<TSharedPtrTS<IStreamSegment>>& OutAlreadyEndedStreams)
+{
+	OutAlreadyEndedStreams.Empty();
+	if (LoaderType == FStreamSegmentRequestISOBMFF::ELoaderType::Muxed && PrimarySampleFileOffsetStart < 0)
+	{
+		for(auto& Stream : DependentStreams)
+		{
+			OutAlreadyEndedStreams.Emplace(Stream);
+		}
+	}
+}
+
+FTimeRange FStreamSegmentRequestISOBMFF::GetTimeRange() const
+{
+	FTimeRange tr;
+	tr.Start = FirstDTS;
+	tr.End = tr.Start + SegmentDuration;
+	tr.Start.SetSequenceIndex(TimestampSequenceIndex);
+	tr.End.SetSequenceIndex(TimestampSequenceIndex);
+	return tr;
+}
+
+
+/*********************************************************************************************************************/
+/*********************************************************************************************************************/
+/*********************************************************************************************************************/
+
+
+FStreamReaderISOBMFF::~FStreamReaderISOBMFF()
+{
+	Close();
+}
+
+void FStreamReaderISOBMFF::SetLoaderType(FStreamSegmentRequestISOBMFF::ELoaderType InLoaderType)
+{
+	LoaderType = InLoaderType;
+}
+
+bool FStreamReaderISOBMFF::LoadsMultiplexedTracksSequentially()
+{
+	return LoaderType == FStreamSegmentRequestISOBMFF::ELoaderType::Muxed;
+}
+
+UEMediaError FStreamReaderISOBMFF::Create(IPlayerSessionServices* InPlayerSessionService, const CreateParam &InCreateParam)
+{
+	check(InPlayerSessionService);
+	PlayerSessionService = InPlayerSessionService;
+
+	if (!InCreateParam.MemoryProvider || !InCreateParam.EventListener)
+	{
+		return UEMEDIA_ERROR_BAD_ARGUMENTS;
+	}
+
+	bIsStarted = true;
+	if (LoaderType == FStreamSegmentRequestISOBMFF::ELoaderType::Chunks || LoaderType == FStreamSegmentRequestISOBMFF::ELoaderType::Fragmented)
+	{
+		for(int32 i=0; i<FMEDIA_STATIC_ARRAY_COUNT(StreamHandlers); ++i)
+		{
+			StreamHandlers[i].PlayerSessionServices = PlayerSessionService;
+			StreamHandlers[i].Parameters = InCreateParam;
+			StreamHandlers[i].bTerminate = false;
+			StreamHandlers[i].bWasStarted = false;
+			StreamHandlers[i].bRequestCanceled = false;
+			StreamHandlers[i].bSilentCancellation = false;
+			StreamHandlers[i].bHasErrored  = false;
+			StreamHandlers[i].IsIdleSignal.Signal();
+			StreamHandlers[i].ThreadSetName(i==0 ? "Electra ISOBMFF video loader" :
+											i==1 ? "Electra ISOBMFF audio loader" :
+												   "Electra ISOBMFF subtitle loader");
+		}
+	}
+	else
+	{
+		StreamHandlers[0].PlayerSessionServices = PlayerSessionService;
+		StreamHandlers[0].Parameters = InCreateParam;
+		StreamHandlers[0].bTerminate = false;
+		StreamHandlers[0].bWasStarted = false;
+		StreamHandlers[0].bRequestCanceled = false;
+		StreamHandlers[0].bSilentCancellation = false;
+		StreamHandlers[0].bHasErrored  = false;
+		StreamHandlers[0].IsIdleSignal.Signal();
+		StreamHandlers[0].ThreadSetName("Electra ISOBMFF loader");
+	}
+	return UEMEDIA_ERROR_OK;
+}
+
+
+void FStreamReaderISOBMFF::Close()
+{
+	if (bIsStarted)
+	{
+		bIsStarted = false;
+		// Signal the worker threads to end.
+		for(int32 i=0; i<FMEDIA_STATIC_ARRAY_COUNT(StreamHandlers); ++i)
+		{
+			StreamHandlers[i].bTerminate = true;
+			StreamHandlers[i].Cancel(true);
+			StreamHandlers[i].SignalWork();
+		}
+		// Wait until they finished.
+		for(int32 i=0; i<FMEDIA_STATIC_ARRAY_COUNT(StreamHandlers); ++i)
+		{
+			if (StreamHandlers[i].bWasStarted)
+			{
+				StreamHandlers[i].ThreadWaitDone();
+				StreamHandlers[i].ThreadReset();
+			}
+		}
+	}
+}
+
+
+void FStreamReaderISOBMFF::FStreamHandler::LogMessage(IInfoLog::ELevel Level, const FString& Message)
+{
+	PlayerSessionServices->PostLog(Facility::EFacility::ISOBMFFStreamReader, Level, Message);
+}
+
+IStreamReader::EAddResult FStreamReaderISOBMFF::AddRequest(uint32 CurrentPlaybackSequenceID, TSharedPtrTS<IStreamSegment> InRequest)
+{
+	TSharedPtrTS<FStreamSegmentRequestISOBMFF> Request = StaticCastSharedPtr<FStreamSegmentRequestISOBMFF>(InRequest);
+	if ((LoaderType == FStreamSegmentRequestISOBMFF::ELoaderType::Chunks || LoaderType == FStreamSegmentRequestISOBMFF::ELoaderType::Fragmented) && !Request->DependentStreams.IsEmpty())
+	{
+		PostError(PlayerSessionService, TEXT("Initial start request segments cannot be enqueued!"), 0);
+		return IStreamReader::EAddResult::Error;
+	}
+	// Do not add empty mux requests.
+	else if (LoaderType != FStreamSegmentRequestISOBMFF::ELoaderType::Muxed || Request->PrimarySampleFileOffsetStart >= 0)
+	{
+		// Get the handler for the main request.
+		FStreamHandler* Handler = nullptr;
+		if (LoaderType == FStreamSegmentRequestISOBMFF::ELoaderType::Chunks || LoaderType == FStreamSegmentRequestISOBMFF::ELoaderType::Fragmented)
+		{
+			switch(Request->GetType())
+			{
+				case EStreamType::Video:
+				{
+					Handler = &StreamHandlers[0];
+					break;
+				}
+				case EStreamType::Audio:
+				{
+					Handler = &StreamHandlers[1];
+					break;
+				}
+				case EStreamType::Subtitle:
+				{
+					Handler = &StreamHandlers[2];
+					break;
+				}
+				default:
+				{
+					break;
+				}
+			}
+		}
+		else
+		{
+			Handler = &StreamHandlers[0];
+		}
+		if (!Handler)
+		{
+			ErrorDetail.SetMessage(FString::Printf(TEXT("No handler for stream type")));
+			return IStreamReader::EAddResult::Error;
+		}
+		// Is the handler busy?
+		bool bIsIdle = Handler->IsIdleSignal.WaitTimeout(1000 * 1000);
+		if (!bIsIdle)
+		{
+			ErrorDetail.SetMessage(FString::Printf(TEXT("The handler for this stream type is busy!?")));
+			return IStreamReader::EAddResult::Error;
+		}
+
+		Request->SetPlaybackSequenceID(CurrentPlaybackSequenceID);
+		if (!Handler->bWasStarted)
+		{
+			Handler->ThreadStart(FMediaRunnable::FStartDelegate::CreateRaw(Handler, &FStreamHandler::WorkerThread));
+			Handler->bWasStarted = true;
+		}
+
+		Handler->bRequestCanceled = false;
+		Handler->bSilentCancellation = false;
+		Handler->bHasErrored = false;
+		Handler->CurrentRequest = Request;
+		Handler->SignalWork();
+	}
+	return EAddResult::Added;
+}
+
+void FStreamReaderISOBMFF::CancelRequest(EStreamType StreamType, bool bSilent)
+{
+	if (LoaderType == FStreamSegmentRequestISOBMFF::ELoaderType::Chunks || LoaderType == FStreamSegmentRequestISOBMFF::ELoaderType::Fragmented)
+	{
+		if (StreamType == EStreamType::Video)
+		{
+			StreamHandlers[0].Cancel(bSilent);
+		}
+		else if (StreamType == EStreamType::Audio)
+		{
+			StreamHandlers[1].Cancel(bSilent);
+		}
+		else if (StreamType == EStreamType::Subtitle)
+		{
+			StreamHandlers[2].Cancel(bSilent);
+		}
+	}
+	else
+	{
+		StreamHandlers[0].Cancel(bSilent);
+	}
+}
+
+void FStreamReaderISOBMFF::CancelRequests()
+{
+	for(int32 i=0; i<FMEDIA_STATIC_ARRAY_COUNT(StreamHandlers); ++i)
+	{
+		StreamHandlers[i].Cancel(false);
+	}
+}
+
+
+std::atomic<uint32> FStreamReaderISOBMFF::FStreamHandler::UniqueDownloadID;
+
+void FStreamReaderISOBMFF::FStreamHandler::Cancel(bool bSilent)
+{
+	bSilentCancellation = bSilent;
+	bRequestCanceled = true;
+}
+
+void FStreamReaderISOBMFF::FStreamHandler::SignalWork()
+{
+	WorkSignal.Release();
+}
+
+
+bool FStreamReaderISOBMFF::FStreamHandler::HasBeenAborted() const
+{
+	return bRequestCanceled;
+}
+
+bool FStreamReaderISOBMFF::FStreamHandler::HasErrored() const
+{
+	return bHasErrored;
+}
+
+int32 FStreamReaderISOBMFF::FStreamHandler::HTTPProgressCallback(const IElectraHttpManager::FRequest* InRequest)
+{
+	// Aborted?
+	return HasBeenAborted() ? 1 : 0;
+}
+
+void FStreamReaderISOBMFF::FStreamHandler::HTTPCompletionCallback(const IElectraHttpManager::FRequest* InRequest)
+{
+	bHasErrored = InRequest->ConnectionInfo.StatusInfo.ErrorDetail.IsError();
+}
+
+
+void FStreamReaderISOBMFF::FStreamHandler::WorkerThread()
+{
+	LLM_SCOPE(ELLMTag::ElectraPlayer);
+
+	while(!bTerminate)
+	{
+		WorkSignal.Obtain();
+		if (!bTerminate)
+		{
+			if (CurrentRequest.IsValid())
+			{
+				IsIdleSignal.Reset();
+				if (!bRequestCanceled)
+				{
+					if (CurrentRequest->LoaderType == FStreamSegmentRequestISOBMFF::ELoaderType::Chunks || CurrentRequest->LoaderType == FStreamSegmentRequestISOBMFF::ELoaderType::Fragmented)
+					{
+						HandleChunkOrFragmentRequest();
+					}
+					else if (CurrentRequest->LoaderType == FStreamSegmentRequestISOBMFF::ELoaderType::Muxed)
+					{
+						HandleMuxedRequest();
+					}
+					else
+					{
+						unimplemented();
+						CurrentRequest.Reset();
+					}
+				}
+				else
+				{
+					CurrentRequest.Reset();
+				}
+				IsIdleSignal.Signal();
+			}
+			bRequestCanceled = false;
+			bSilentCancellation = false;
+			bHasErrored = false;
+		}
+	}
+}
+
+
+void FStreamReaderISOBMFF::FStreamHandler::HandleMuxedRequest()
+{
+	TSharedPtrTS<FStreamSegmentRequestISOBMFF> Request = CurrentRequest;
+
+	// Create CSD and buffer info for all the streams.
+	struct FTrackVars
+	{
+		TSharedPtr<FAccessUnit::CodecData, ESPMode::ThreadSafe> CSD;
+		TSharedPtr<FBufferSourceInfo, ESPMode::ThreadSafe> BufferInfo;
+		TArray<FAccessUnit*> AccessUnits;
+		FTimeValue DurationSuccessfullyRead { FTimeValue::GetZero() };
+		FTimeValue DurationSuccessfullyDelivered { FTimeValue::GetZero() };
+		EStreamType StreamType { EStreamType::Unsupported };
+		bool bIsFirstInSequence { true };
+		bool bEndOfTrack { false };
+		bool bDeliveredLastAccessUnit { false };
+		bool bNotifiedOfTrackEnd { false };
+		int32 DependendStreamIdx { -1 };
+	};
+	TMap<uint32, FTrackVars> TrackMap;
+	for(int32 i=0; i<Request->DependentStreams.Num(); ++i)
+	{
+		TSharedPtr<FAccessUnit::CodecData, ESPMode::ThreadSafe> CSD(new FAccessUnit::CodecData);
+		CSD->CodecSpecificData = Request->DependentStreams[i]->Track.TrackInfo->CodecFormat.CSD;
+		CSD->RawCSD = Request->DependentStreams[i]->Track.TrackInfo->CodecFormat.DCR;
+		CSD->ParsedInfo = FStreamCodecInformation(Request->DependentStreams[i]->Track.TrackInfo->CodecFormat);
+		TSharedPtr<FBufferSourceInfo, ESPMode::ThreadSafe> BufferInfo = MakeShared<FBufferSourceInfo, ESPMode::ThreadSafe>();
+		BufferInfo->PlaybackSequenceID = Request->GetPlaybackSequenceID();
+		BufferInfo->PeriodID = Request->AssetID;
+		BufferInfo->PeriodAdaptationSetID = Request->AssetID + TEXT(".") + Request->DependentStreams[i]->AdaptationSetID;
+		BufferInfo->LanguageTag = Request->DependentStreams[i]->LanguageTag;
+		BufferInfo->HardIndex = Request->DependentStreams[i]->AdaptationSetTrackIndex;
+		BufferInfo->Codec = CSD->ParsedInfo.GetCodecName();
+		BufferInfo->Kind = Request->DependentStreams[i]->KindOfTrack;
+		FTrackVars tv;
+		tv.DependendStreamIdx = i;
+		tv.CSD = MoveTemp(CSD);
+		tv.BufferInfo = MoveTemp(BufferInfo);
+		tv.StreamType = Request->DependentStreams[i]->StreamType;
+		// If there is no iterator for this track then it is already at EOT.
+		tv.bEndOfTrack = !Request->DependentStreams[i]->Track.TrackIterator;
+		// If this track is at its end we let the player know that this buffer will not be
+		// receiving any data.
+		if (tv.bEndOfTrack)
+		{
+			tv.bNotifiedOfTrackEnd = true;
+			Parameters.EventListener->OnFragmentReachedEOS(tv.StreamType, tv.BufferInfo);
+		}
+		TrackMap.Emplace(Request->DependentStreams[i]->Track.TrackInfo->TrackID, MoveTemp(tv));
+	}
+
+	Metrics::FSegmentDownloadStats& ds = Request->DownloadStats;
+	ds.StatsID = ++UniqueDownloadID;
+	ds.MediaAssetID = Request->AssetID;
+	ds.AdaptationSetID = Request->AdaptationSetID;
+	ds.RepresentationID = Request->RepresentationID;
+	ds.Bitrate = Request->Bitrate;
+	ds.FailureReason.Empty();
+	ds.bWasSuccessful = true;
+	ds.bWasAborted = false;
+	ds.bDidTimeout = false;
+	ds.HTTPStatusCode = 0;
+	ds.StreamType = Request->GetType();
+	ds.SegmentType = Metrics::ESegmentType::Media;
+	ds.PresentationTime = Request->EarliestPTS.GetAsSeconds();
+	ds.Duration = Request->SegmentDuration.GetAsSeconds();
+	ds.DurationDownloaded = 0.0;
+	ds.DurationDelivered = 0.0;
+	ds.TimeToFirstByte = 0.0;
+	ds.TimeToDownload = 0.0;
+	ds.ByteSize = -1;
+	ds.NumBytesDownloaded = 0;
+	ds.bInsertedFillerData = false;
+	ds.URL.URL = Request->MediaURL;
+	ds.bIsMissingSegment = false;
+	ds.bParseFailure = false;
+	ds.RetryNumber = Request->NumOverallRetries;
+
+	Parameters.EventListener->OnFragmentOpen(Request);
+
+	TSharedPtrTS<IElectraHttpManager::FProgressListener> ProgressListener = MakeSharedTS<IElectraHttpManager::FProgressListener>();
+	ProgressListener->CompletionDelegate = IElectraHttpManager::FProgressListener::FCompletionDelegate::CreateRaw(this, &FStreamReaderISOBMFF::FStreamHandler::HTTPCompletionCallback);
+	ProgressListener->ProgressDelegate = IElectraHttpManager::FProgressListener::FProgressDelegate::CreateRaw(this, &FStreamReaderISOBMFF::FStreamHandler::HTTPProgressCallback);
+
+	const int64 StartOffset = Request->PrimarySampleFileOffsetStart;
+	const int64 EndOffset = Request->PrimarySampleFileOffsetEnd;
+
+	ReceiveBuffer.Reset();
+	ReceiveBuffer = MakeSharedTS<FWaitableBuffer>();
+	TSharedPtrTS<IElectraHttpManager::FRequest> HTTP(new IElectraHttpManager::FRequest);
+	HTTP->Parameters.URL = Request->MediaURL;
+	HTTP->Parameters.Range.Start = StartOffset;
+	HTTP->Parameters.Range.EndIncluding = EndOffset - 1;
+	HTTP->Parameters.AcceptEncoding.Set(TEXT("identity"));
+	HTTP->Parameters.ConnectTimeout = PlayerSessionServices->GetOptionValue(ISOBMFF::OptionKeyISOBMFFLoadConnectTimeout).SafeGetTimeValue(FTimeValue().SetFromMilliseconds(1000 * 8));
+	HTTP->Parameters.NoDataTimeout = PlayerSessionServices->GetOptionValue(ISOBMFF::OptionKeyISOBMFFLoadNoDataTimeout).SafeGetTimeValue(FTimeValue().SetFromMilliseconds(1000 * 6));
+	HTTP->ReceiveBuffer = ReceiveBuffer;
+	HTTP->ProgressListener = ProgressListener;
+	HTTP->ResponseCache = PlayerSessionServices->GetHTTPResponseCache();
+	PlayerSessionServices->GetHTTPManager()->AddRequest(HTTP, false);
+
+	bool bDone = false;
+
+	const uint32 PlaybackSequenceID = Request->GetPlaybackSequenceID();
+
+	// Locate the track iterator that will handle the next incoming data.
+	auto GetNextTrackIndexToIterate = [&]() -> int32
+	{
+		int64 NextIteratorFileDataOffset = TNumericLimits<int64>::Max();
+		int32 ItIdx = -1;
+		for(int32 i=0; i<Request->DependentStreams.Num(); ++i)
+		{
+			if (Request->DependentStreams[i]->Track.TrackIterator && Request->DependentStreams[i]->Track.TrackIterator->GetSampleFileOffset() < NextIteratorFileDataOffset)
+			{
+				NextIteratorFileDataOffset = Request->DependentStreams[i]->Track.TrackIterator->GetSampleFileOffset();
+				ItIdx = i;
+			}
+		}
+		return ItIdx;
+	};
+
+	auto ShouldStop = [&]() -> bool
+	{
+		return HasErrored() || HasBeenAborted() || bTerminate;
+	};
+
+
+	auto DiscardAllPendingAccessUnits = [&]() -> void
+	{
+		for(auto &tvit : TrackMap)
+		{
+			FTrackVars& tv(tvit.Value);
+			while(!tv.AccessUnits.IsEmpty())
+			{
+				FAccessUnit* AU = tv.AccessUnits.Pop();
+				FAccessUnit::Release(AU);
+				AU = nullptr;
+			}
+		}
+	};
+
+	auto DeliverAccessUnits = [&](bool bUntilDone) -> void
+	{
+		while(!HasBeenAborted() && !bTerminate)
+		{
+			int32 NumBlocked = 0;
+			int32 NumWithAUs = 0;
+			for(auto &tvit : TrackMap)
+			{
+				FTrackVars& tv(tvit.Value);
+				if (tv.AccessUnits.Num())
+				{
+					if (!tv.bDeliveredLastAccessUnit)
+					{
+						++NumWithAUs;
+						FAccessUnit* AU = tv.AccessUnits[0];
+						if (Parameters.EventListener->OnFragmentAccessUnitReceived(AU))
+						{
+							tv.DurationSuccessfullyDelivered += AU->Duration;
+							tv.bDeliveredLastAccessUnit |= AU->bIsLastInPeriod;
+							tv.AccessUnits.RemoveAt(0);
+							AU = nullptr;
+						}
+						else
+						{
+							++NumBlocked;
+						}
+					}
+					else
+					{
+						while(!tv.AccessUnits.IsEmpty())
+						{
+							FAccessUnit* AU = tv.AccessUnits.Pop();
+							FAccessUnit::Release(AU);
+							AU = nullptr;
+						}
+					}
+				}
+				else if (!tv.bNotifiedOfTrackEnd && !Request->DependentStreams[tv.DependendStreamIdx]->Track.TrackIterator)
+				{
+					tv.bNotifiedOfTrackEnd = true;
+					Parameters.EventListener->OnFragmentReachedEOS(tv.StreamType, tv.BufferInfo);
+				}
+			}
+			if (!bUntilDone || NumWithAUs == 0)
+			{
+				break;
+			}
+			if (NumBlocked == NumWithAUs)
+			{
+				FMediaRunnable::SleepMicroseconds(1000 * 10);
+			}
+		}
+	};
+
+	int64 SizeNeededInBuffer = 0;
+	int32 NextTrackIndex = -1;
+	while(!bDone)
+	{
+		if (NextTrackIndex < 0)
+		{
+			NextTrackIndex = GetNextTrackIndexToIterate();
+			// All at EOS?
+			if (NextTrackIndex < 0)
+			{
+				Request->bIsLastSegment = true;
+				bDone = true;
+				break;
+			}
+		}
+		TSharedPtrTS<FStreamSegmentRequestISOBMFF>& DepStr(Request->DependentStreams[NextTrackIndex]);
+
+		FTrackVars& TrackVars(TrackMap[DepStr->Track.TrackInfo->TrackID]);
+		check(DepStr->Track.TrackChunkSampleInfo.NumSamples);
+
+		const int64 SampleRelativeOffset = DepStr->Track.TrackIterator->GetSampleFileOffset() - StartOffset;
+		const int64 SampleSize = DepStr->Track.TrackIterator->GetSampleSize();
+		SizeNeededInBuffer = SampleRelativeOffset + SampleSize;
+
+		// Wait for the data to arrive. While it does emit already accumulated access units.
+		while(!ShouldStop())
+		{
+			if (ReceiveBuffer->WaitUntilSizeAvailable(SizeNeededInBuffer, 1000 * 100))
+			{
+				// Did we get what we need or did we hit EOF?
+				if (ReceiveBuffer->GetLinearReadSize() < SizeNeededInBuffer)
+				{
+					bDone = true;
+				}
+				break;
+			}
+			else
+			{
+				// If the read buffer doesn't have the required size yet we try to deliver accumulated access units
+				DeliverAccessUnits(false);
+			}
+		}
+		if (bDone || ShouldStop())
+		{
+			break;
+		}
+
+		// Get the sample properties
+		auto TrackIt = DepStr->Track.TrackIterator;
+		uint32 SampleNumber = TrackIt->GetSampleNumber();
+		MP4Utilities::FFractionalTime DTS = TrackIt->GetEffectiveDTS();
+		MP4Utilities::FFractionalTime PTS = TrackIt->GetEffectivePTS();
+		MP4Utilities::FFractionalTime Duration = TrackIt->GetDuration();
+		bool bIsSyncSample = TrackIt->IsSyncOrRAPSample();
+
+		// Create an access unit.
+		FAccessUnit* AccessUnit = FAccessUnit::Create(Parameters.MemoryProvider);
+		if (AccessUnit)
+		{
+			const int64 TimestampSequenceIndex = Request->TimestampSequenceIndex;
+			AccessUnit->ESType = DepStr->StreamType;
+			AccessUnit->PTS.SetFromND(PTS.GetNumerator(), PTS.GetDenominator(), TimestampSequenceIndex);
+			AccessUnit->DTS.SetFromND(DTS.GetNumerator(), DTS.GetDenominator(), TimestampSequenceIndex);
+			AccessUnit->Duration.SetFromND(Duration.GetNumerator(), Duration.GetDenominator());
+
+			AccessUnit->SequenceIndex = TimestampSequenceIndex;
+			AccessUnit->EarliestPTS = DepStr->EarliestPTS;
+			AccessUnit->EarliestPTS.SetSequenceIndex(TimestampSequenceIndex);
+			AccessUnit->LatestPTS = DepStr->LastPTS;
+			AccessUnit->LatestPTS.SetSequenceIndex(TimestampSequenceIndex);
+
+			AccessUnit->AUSize = (uint32) SampleSize;
+			AccessUnit->AUCodecData = TrackVars.CSD;
+			AccessUnit->BufferSourceInfo = TrackVars.BufferInfo;
+
+			AccessUnit->bIsFirstInSequence = TrackVars.bIsFirstInSequence;
+			TrackVars.bIsFirstInSequence = false;
+			AccessUnit->bIsSyncSample = bIsSyncSample;
+			AccessUnit->bIsDummyData = false;
+			AccessUnit->AUData = AccessUnit->AllocatePayloadBuffer(AccessUnit->AUSize);
+			{
+				FScopeLock lock(ReceiveBuffer->GetLock());
+				FMemory::Memcpy(AccessUnit->AUData, reinterpret_cast<const uint8*>(ReceiveBuffer->GetLinearReadData()) + SampleRelativeOffset, SampleSize);
+			}
+			if (AccessUnit->DTS >= AccessUnit->LatestPTS && AccessUnit->PTS >= AccessUnit->LatestPTS)
+			{
+				AccessUnit->bIsLastInPeriod = true;
+				TrackVars.bEndOfTrack = true;
+			}
+			TrackVars.DurationSuccessfullyRead += AccessUnit->Duration;
+			TrackVars.AccessUnits.Emplace(AccessUnit);
+			AccessUnit = nullptr;
+		}
+		DeliverAccessUnits(false);
+
+		// Exhausted all samples in this chunk?
+		const bool bEOT = !TrackIt->Next();
+		if (bEOT)
+		{
+			NextTrackIndex = -1;
+			TrackVars.bEndOfTrack = true;
+			// Dump the iterator so we know that this track is done.
+			DepStr->Track.TrackIterator.Reset();
+		}
+		else if (--DepStr->Track.TrackChunkSampleInfo.NumSamples == 0)
+		{
+			TrackIt->GetCurrentChunkRemainingSampleInfo(DepStr->Track.TrackChunkSampleInfo);
+			NextTrackIndex = -1;
+		}
+	}
+	// Remove the download request.
+	ProgressListener.Reset();
+	PlayerSessionServices->GetHTTPManager()->RemoveRequest(HTTP, false);
+	Request->ConnectionInfo = HTTP->ConnectionInfo;
+	HTTP.Reset();
+	ReceiveBuffer.Reset();
+
+	// Deliver all pending access units unless we get canceled.
+	DeliverAccessUnits(true);
+	// Anything that's still there we now need to discard.
+	DiscardAllPendingAccessUnits();
+	// If all tracks reached their end we mark the request as being the last one.
+	bool bAllAtLastSample = true;
+	for(auto &tvit : TrackMap)
+	{
+		FTrackVars& tv(tvit.Value);
+		if (!tv.bEndOfTrack)
+		{
+			bAllAtLastSample = false;
+		}
+	}
+	Request->bIsLastSegment |= bAllAtLastSample;
+
+	// Set up download stat fields.
+	ds.HTTPStatusCode = Request->ConnectionInfo.StatusInfo.HTTPStatus;
+	ds.TimeToFirstByte = Request->ConnectionInfo.TimeUntilFirstByte;
+	ds.TimeToDownload = (Request->ConnectionInfo.RequestEndTime - Request->ConnectionInfo.RequestStartTime).GetAsSeconds();
+	ds.ByteSize = Request->ConnectionInfo.ContentLength;
+	ds.NumBytesDownloaded = Request->ConnectionInfo.BytesReadSoFar;
+	ds.FailureReason = Request->ConnectionInfo.StatusInfo.ErrorDetail.GetMessage();
+	ds.bWasSuccessful = !bHasErrored;
+	ds.URL.URL = Request->ConnectionInfo.EffectiveURL;
+	if (Request->ReportingIndexInMux >= 0)
+	{
+		const FTrackVars& tv(TrackMap[Request->DependentStreams[Request->ReportingIndexInMux]->Track.TrackInfo->TrackID]);
+		ds.DurationDownloaded = tv.DurationSuccessfullyRead.GetAsSeconds();
+		ds.DurationDelivered  = tv.DurationSuccessfullyDelivered.GetAsSeconds();
+	}
+	ds.bIsCachedResponse = Request->ConnectionInfo.bIsCachedResponse;
+
+	// Reset the current request so another one can be added immediately when we call OnFragmentClose()
+	CurrentRequest.Reset();
+	if (!bSilentCancellation)
+	{
+		PlayerSessionServices->GetStreamSelector()->ReportDownloadEnd(ds);
+		Parameters.EventListener->OnFragmentClose(Request);
+	}
+}
+
+void FStreamReaderISOBMFF::FStreamHandler::HandleChunkOrFragmentRequest()
+{
+	TSharedPtrTS<FStreamSegmentRequestISOBMFF> Request = CurrentRequest;
+
+	Metrics::FSegmentDownloadStats& ds = Request->DownloadStats;
+	ds.StatsID = ++UniqueDownloadID;
+	ds.MediaAssetID = Request->AssetID;
+	ds.AdaptationSetID = Request->AdaptationSetID;
+	ds.RepresentationID = Request->RepresentationID;
+	ds.Bitrate = Request->Bitrate;
+	ds.FailureReason.Empty();
+	ds.bWasSuccessful = true;
+	ds.bWasAborted = false;
+	ds.bDidTimeout = false;
+	ds.HTTPStatusCode = 0;
+	ds.StreamType = Request->GetType();
+	ds.SegmentType = Metrics::ESegmentType::Media;
+	ds.PresentationTime = Request->EarliestPTS.GetAsSeconds();
+	ds.Duration = Request->SegmentDuration.GetAsSeconds();
+	ds.DurationDownloaded = 0.0;
+	ds.DurationDelivered = 0.0;
+	ds.TimeToFirstByte = 0.0;
+	ds.TimeToDownload = 0.0;
+	ds.ByteSize = -1;
+	ds.NumBytesDownloaded = 0;
+	ds.bInsertedFillerData = false;
+	ds.URL.URL = Request->MediaURL;
+	ds.bIsMissingSegment = false;
+	ds.bParseFailure = false;
+	ds.RetryNumber = Request->NumOverallRetries;
+
+	Parameters.EventListener->OnFragmentOpen(Request);
+
+	TSharedPtrTS<IElectraHttpManager::FProgressListener> ProgressListener = MakeSharedTS<IElectraHttpManager::FProgressListener>();
+	ProgressListener->CompletionDelegate = IElectraHttpManager::FProgressListener::FCompletionDelegate::CreateRaw(this, &FStreamReaderISOBMFF::FStreamHandler::HTTPCompletionCallback);
+	ProgressListener->ProgressDelegate = IElectraHttpManager::FProgressListener::FProgressDelegate::CreateRaw(this, &FStreamReaderISOBMFF::FStreamHandler::HTTPProgressCallback);
+
+	const int64 StartOffset = Request->Track.TrackIterator->GetSampleFileOffset();
+	const int64 EndOffset = StartOffset + Request->Track.TrackChunkSampleInfo.SizeInBytes;
+
+	ReceiveBuffer.Reset();
+	ReceiveBuffer = MakeSharedTS<FWaitableBuffer>();
+	TSharedPtrTS<IElectraHttpManager::FRequest> HTTP(new IElectraHttpManager::FRequest);
+	HTTP->Parameters.URL = Request->MediaURL;
+	HTTP->Parameters.Range.Start = StartOffset;
+	HTTP->Parameters.Range.EndIncluding = EndOffset - 1;
+	HTTP->Parameters.AcceptEncoding.Set(TEXT("identity"));
+	HTTP->Parameters.ConnectTimeout = PlayerSessionServices->GetOptionValue(ISOBMFF::OptionKeyISOBMFFLoadConnectTimeout).SafeGetTimeValue(FTimeValue().SetFromMilliseconds(1000 * 8));
+	HTTP->Parameters.NoDataTimeout = PlayerSessionServices->GetOptionValue(ISOBMFF::OptionKeyISOBMFFLoadNoDataTimeout).SafeGetTimeValue(FTimeValue().SetFromMilliseconds(1000 * 6));
+	HTTP->ReceiveBuffer = ReceiveBuffer;
+	HTTP->ProgressListener = ProgressListener;
+	HTTP->ResponseCache = PlayerSessionServices->GetHTTPResponseCache();
+	PlayerSessionServices->GetHTTPManager()->AddRequest(HTTP, false);
+
+	bool bDone = false;
+
+	// Create CSD and buffer info
+	TSharedPtrTS<FAccessUnit::CodecData> CSD(new FAccessUnit::CodecData);
+	CSD->CodecSpecificData = Request->Track.TrackInfo->CodecFormat.CSD;
+	CSD->RawCSD = Request->Track.TrackInfo->CodecFormat.DCR;
+	CSD->ParsedInfo = FStreamCodecInformation(Request->Track.TrackInfo->CodecFormat);
+	TSharedPtr<FBufferSourceInfo, ESPMode::ThreadSafe> BufferSource = MakeShared<FBufferSourceInfo, ESPMode::ThreadSafe>();
+	BufferSource->PlaybackSequenceID = Request->GetPlaybackSequenceID();
+	BufferSource->PeriodID = Request->AssetID;
+	BufferSource->PeriodAdaptationSetID = Request->AssetID + TEXT(".") + Request->AdaptationSetID;
+	BufferSource->LanguageTag = Request->LanguageTag;
+	BufferSource->HardIndex = Request->AdaptationSetTrackIndex;
+	BufferSource->Codec = CSD->ParsedInfo.GetCodecName();
+	BufferSource->Kind = Request->KindOfTrack;
+
+	uint32 PlaybackSequenceID = Request->GetPlaybackSequenceID();
+	auto TrackIt = Request->Track.TrackIterator;
+	int64 SizeNeededInBuffer = 0;
+	FTimeValue DurationSuccessfullyRead(FTimeValue::GetZero());
+	FTimeValue DurationSuccessfullyDelivered(FTimeValue::GetZero());
+	bool bIsFirstInSequence = true;
+	bool bSentLastAccessUnit = false;
+	while(!bDone && !HasErrored() && !HasBeenAborted() && !bTerminate)
+	{
+		// Get the number of bytes we need for the next sample.
+		int64 SampleSize = TrackIt->GetSampleSize();
+		SizeNeededInBuffer = TrackIt->GetSampleFileOffset() + SampleSize - StartOffset;
+		// Wait until the required amount of data has been read.
+		while(!HasErrored() && !HasBeenAborted() && !bTerminate)
+		{
+			// Check periodically if we are to abort in case we do not have enough data yet.
+			if (ReceiveBuffer->WaitUntilSizeAvailable(SizeNeededInBuffer, 1000 * 100))
+			{
+				break;
+			}
+		}
+		// Leave the loop if we have to abort.
+		if (!(!HasErrored() && !HasBeenAborted() && !bTerminate))
+		{
+			break;
+		}
+		// Did we get enough data?
+		if (ReceiveBuffer->GetLinearReadSize() < SizeNeededInBuffer)
+		{
+			bDone = true;
+			break;
+		}
+
+
+		// Get the sample properties
+		uint32 SampleNumber = TrackIt->GetSampleNumber();
+		MP4Utilities::FFractionalTime DTS = TrackIt->GetEffectiveDTS();
+		MP4Utilities::FFractionalTime PTS = TrackIt->GetEffectivePTS();
+		MP4Utilities::FFractionalTime Duration = TrackIt->GetDuration();
+		bool bIsSyncSample = TrackIt->IsSyncOrRAPSample();
+
+		// Create an access unit.
+		FAccessUnit* AccessUnit = FAccessUnit::Create(Parameters.MemoryProvider);
+		if (AccessUnit)
+		{
+			AccessUnit->ESType = Request->StreamType;
+			AccessUnit->PTS.SetFromND(PTS.GetNumerator(), PTS.GetDenominator(), Request->TimestampSequenceIndex);
+			AccessUnit->DTS.SetFromND(DTS.GetNumerator(), DTS.GetDenominator(), Request->TimestampSequenceIndex);
+			AccessUnit->Duration.SetFromND(Duration.GetNumerator(), Duration.GetDenominator());
+
+			AccessUnit->SequenceIndex = Request->TimestampSequenceIndex;
+			AccessUnit->EarliestPTS = Request->EarliestPTS;
+			AccessUnit->EarliestPTS.SetSequenceIndex(Request->TimestampSequenceIndex);
+			AccessUnit->LatestPTS = Request->LastPTS;
+			AccessUnit->LatestPTS.SetSequenceIndex(Request->TimestampSequenceIndex);
+
+			AccessUnit->AUSize = (uint32) SampleSize;
+			AccessUnit->AUCodecData = CSD;
+			AccessUnit->BufferSourceInfo = BufferSource;
+
+			AccessUnit->bIsFirstInSequence = bIsFirstInSequence;
+			bIsFirstInSequence = false;
+			AccessUnit->bIsSyncSample = bIsSyncSample;
+			AccessUnit->bIsDummyData = false;
+			AccessUnit->AUData = AccessUnit->AllocatePayloadBuffer(AccessUnit->AUSize);
+			{
+				const int64 SampleBufferOffset = TrackIt->GetSampleFileOffset() - StartOffset;
+				FScopeLock lock(ReceiveBuffer->GetLock());
+				FMemory::Memcpy(AccessUnit->AUData, reinterpret_cast<const uint8*>(ReceiveBuffer->GetLinearReadData()) + SampleBufferOffset, SampleSize);
+			}
+
+			if (AccessUnit->DTS >= AccessUnit->LatestPTS && AccessUnit->PTS >= AccessUnit->LatestPTS)
+			{
+				// Tag the last one and send it off, but stop doing so for the remainder of the segment.
+				// Note: we continue reading this chunk all the way to the end on purpose to avoid prematurely closing the connection.
+				AccessUnit->bIsLastInPeriod = true;
+				Request->bIsLastSegment = true;
+			}
+
+			DurationSuccessfullyRead += AccessUnit->Duration;
+			bool bSentOff = bSentLastAccessUnit;
+			while(!bSentOff && !HasBeenAborted() && !bTerminate)
+			{
+				if (Parameters.EventListener->OnFragmentAccessUnitReceived(AccessUnit))
+				{
+					DurationSuccessfullyDelivered += AccessUnit->Duration;
+					bSentLastAccessUnit |= AccessUnit->bIsLastInPeriod;
+					bSentOff = true;
+					AccessUnit = nullptr;
+				}
+				else
+				{
+					FMediaRunnable::SleepMicroseconds(1000 * 10);
+				}
+			}
+
+			// Release the AU if we still have it.
+			FAccessUnit::Release(AccessUnit);
+			AccessUnit = nullptr;
+		}
+
+		// Advance to the next sample unless we are to abort.
+		if (!HasErrored() && !HasBeenAborted() && !bTerminate)
+		{
+			bool bReachedLastSample = !TrackIt->Next();
+			if (bReachedLastSample)
+			{
+				bDone = true;
+			}
+		}
+	}
+
+	// Remove the download request.
+	ProgressListener.Reset();
+	PlayerSessionServices->GetHTTPManager()->RemoveRequest(HTTP, false);
+	Request->ConnectionInfo = HTTP->ConnectionInfo;
+	HTTP.Reset();
+	ReceiveBuffer.Reset();
+
+	// Set up download stat fields.
+	ds.HTTPStatusCode = Request->ConnectionInfo.StatusInfo.HTTPStatus;
+	ds.TimeToFirstByte = Request->ConnectionInfo.TimeUntilFirstByte;
+	ds.TimeToDownload = (Request->ConnectionInfo.RequestEndTime - Request->ConnectionInfo.RequestStartTime).GetAsSeconds();
+	ds.ByteSize = Request->ConnectionInfo.ContentLength;
+	ds.NumBytesDownloaded = Request->ConnectionInfo.BytesReadSoFar;
+	ds.FailureReason = Request->ConnectionInfo.StatusInfo.ErrorDetail.GetMessage();
+	ds.bWasSuccessful = !bHasErrored;
+	ds.URL.URL = Request->ConnectionInfo.EffectiveURL;
+	ds.DurationDownloaded = DurationSuccessfullyRead.GetAsSeconds();
+	ds.DurationDelivered  = DurationSuccessfullyDelivered.GetAsSeconds();
+	ds.bIsCachedResponse = Request->ConnectionInfo.bIsCachedResponse;
+
+	// Reset the current request so another one can be added immediately when we call OnFragmentClose()
+	CurrentRequest.Reset();
+	if (!bSilentCancellation)
+	{
+		PlayerSessionServices->GetStreamSelector()->ReportDownloadEnd(ds);
+		Parameters.EventListener->OnFragmentClose(Request);
+	}
+}
+} // namespace Electra
diff --git a/Engine/Plugins/Media/ElectraPlayer/Source/ElectraPlayerRuntime/Private/Runtime/Player/isobmff/StreamReaderISOBMFF.h b/Engine/Plugins/Media/ElectraPlayer/Source/ElectraPlayerRuntime/Private/Runtime/Player/isobmff/StreamReaderISOBMFF.h
new file mode 100644
--- /dev/null
+++ b/Engine/Plugins/Media/ElectraPlayer/Source/ElectraPlayerRuntime/Private/Runtime/Player/isobmff/StreamReaderISOBMFF.h
@@ -0,0 +1,175 @@
+// Copyright Epic Games, Inc. All Rights Reserved.
+
+#pragma once
+
+#include "PlayerCore.h"
+#include "Player/Manifest.h"
+#include "Player/PlayerStreamReader.h"
+#include "Player/DRM/DRMManager.h"
+#include "Player/isobmff/PlaylistISOBMFF.h"
+#include "HTTP/HTTPManager.h"
+#include "Utilities/HashFunctions.h"
+#include "Utilities/TimeUtilities.h"
+#include "Utilities/BCP47-Helpers.h"
+#include "SynchronizedClock.h"
+#include "StreamAccessUnitBuffer.h"
+#include "StreamDataBuffer.h"
+#include "MP4Track.h"
+
+#include <atomic>
+
+
+namespace Electra
+{
+
+
+class FStreamSegmentRequestISOBMFF : public IStreamSegment
+{
+public:
+	FStreamSegmentRequestISOBMFF() = default;
+	virtual ~FStreamSegmentRequestISOBMFF() = default;
+
+	void SetPlaybackSequenceID(uint32 InPlaybackSequenceID) override
+	{ PlaybackSequenceID = InPlaybackSequenceID; }
+	uint32 GetPlaybackSequenceID() const override
+	{ return PlaybackSequenceID; }
+	void SetExecutionDelay(const FTimeValue& UTCNow, const FTimeValue& ExecutionDelay) override
+	{ }
+	FTimeValue GetExecuteAtUTCTime() const override
+	{ return FTimeValue::GetInvalid(); }
+	EStreamType GetType() const override
+	{ return StreamType; }
+
+	void GetDependentStreams(TArray<TSharedPtrTS<IStreamSegment>>& OutDependentStreams) const override;
+	void GetRequestedStreams(TArray<TSharedPtrTS<IStreamSegment>>& OutRequestedStreams) override;
+	void GetEndedStreams(TArray<TSharedPtrTS<IStreamSegment>>& OutAlreadyEndedStreams) override;
+
+	FTimeValue GetFirstPTS() const override
+	{ return EarliestPTS; }
+	FTimeRange GetTimeRange() const override;
+
+	int32 GetQualityIndex() const override
+	{ return 0; }
+	int32 GetBitrate() const override
+	{ return Bitrate; }
+
+	void GetDownloadStats(Metrics::FSegmentDownloadStats& OutStats) const override
+	{ OutStats = DownloadStats; }
+	bool GetStartupDelay(FTimeValue& OutStartTime, FTimeValue& OutTimeIntoSegment, FTimeValue& OutSegmentDuration) const override
+	{ return false; }
+
+	struct FTrack
+	{
+		TSharedPtr<FManifestISOBMFFInternal::FTrackInfo, ESPMode::ThreadSafe> TrackInfo;
+		TSharedPtr<MP4Boxes::FMP4Track::FIterator, ESPMode::ThreadSafe> TrackIterator;
+		MP4Boxes::FMP4Track::FIterator::FChunkInfo TrackChunkSampleInfo;
+	};
+	enum class ELoaderType
+	{
+		Chunks,
+		Muxed,
+		Fragmented
+	};
+
+	TSharedPtrTS<ITimelineMediaAsset> MediaAsset;
+	ELoaderType LoaderType { ELoaderType::Chunks };
+	FString AssetID;
+	FString AdaptationSetID;
+	FString RepresentationID;
+	int32 AdaptationSetTrackIndex { 0 };
+	FString MediaURL;
+	BCP47::FLanguageTag LanguageTag;
+	FString KindOfTrack;
+	TArray<TSharedPtrTS<FStreamSegmentRequestISOBMFF>> DependentStreams;
+	FTrack Track;
+	FTimeValue FirstDTS;
+	FTimeValue EarliestPTS;
+	FTimeValue LastPTS;
+	FTimeValue SegmentDuration;
+	EStreamType StreamType {EStreamType::Unsupported};
+	int64 TimestampSequenceIndex { 0 };
+	uint32 PlaybackSequenceID { ~0U };
+	int32 Bitrate { 0 };
+	int32 NumOverallRetries { 0 };
+	bool bIsLastSegment { false };
+
+	// File offsets when reading sequentially all track data in the multiplex
+	int64 PrimarySampleFileOffsetStart { -1 };
+	int64 PrimarySampleFileOffsetEnd { -1 };
+	int32 ReportingIndexInMux { -1 };
+	bool bWasSelectedAtStart { false };
+
+	// Encryption
+	TSharedPtrTS<ElectraCDM::IMediaCDMClient> DrmClient;
+	FString DrmMimeType;
+
+	// Download stats that are set during and after download.
+	Metrics::FSegmentDownloadStats DownloadStats;
+	HTTP::FConnectionInfo ConnectionInfo;
+};
+
+
+/**
+ * This class implements an interface to read from an mp4 file.
+ */
+class FStreamReaderISOBMFF : public IStreamReader
+{
+public:
+	FStreamReaderISOBMFF() = default;
+	virtual ~FStreamReaderISOBMFF();
+	void SetLoaderType(FStreamSegmentRequestISOBMFF::ELoaderType InLoaderType);
+	UEMediaError Create(IPlayerSessionServices* PlayerSessionService, const CreateParam &createParam) override;
+	bool LoadsMultiplexedTracksSequentially() override;
+	void Close() override;
+	EAddResult AddRequest(uint32 CurrentPlaybackSequenceID, TSharedPtrTS<IStreamSegment> Request) override;
+	void CancelRequest(EStreamType StreamType, bool bSilent) override;
+	void CancelRequests() override;
+
+private:
+	ELECTRA_IMPL_DEFAULT_ERROR_METHODS(ISOBMFFStreamReader);
+
+	struct FStreamHandler : public FMediaThread
+	{
+		FStreamHandler() = default;
+		virtual ~FStreamHandler() = default;
+		void Cancel(bool bSilent);
+		void SignalWork();
+		void WorkerThread();
+		void SetError(const FString& InMessage, uint16 InCode);
+		bool HasErrored() const;
+		void LogMessage(IInfoLog::ELevel Level, const FString& Message);
+		int32 HTTPProgressCallback(const IElectraHttpManager::FRequest* Request);
+		void HTTPCompletionCallback(const IElectraHttpManager::FRequest* Request);
+		void HTTPUpdateStats(const FTimeValue& CurrentTime, const IElectraHttpManager::FRequest* Request);
+		void HandleMuxedRequest();
+		void HandleChunkOrFragmentRequest();
+		bool HasBeenAborted() const;
+
+		IStreamReader::CreateParam Parameters;
+		FMediaSemaphore WorkSignal;
+		FMediaEvent IsIdleSignal;
+		TSharedPtrTS<FStreamSegmentRequestISOBMFF> CurrentRequest;
+		IPlayerSessionServices* PlayerSessionServices = nullptr;
+		volatile bool bTerminate = false;
+		volatile bool bWasStarted = false;
+		volatile bool bRequestCanceled = false;
+		volatile bool bSilentCancellation = false;
+		volatile bool bHasErrored = false;
+
+		TSharedPtrTS<FWaitableBuffer> ReceiveBuffer;
+		TSharedPtr<ElectraCDM::IMediaCDMDecrypter, ESPMode::ThreadSafe> Decrypter;
+
+		static std::atomic<uint32> UniqueDownloadID;
+	};
+
+
+	FStreamSegmentRequestISOBMFF::ELoaderType LoaderType { FStreamSegmentRequestISOBMFF::ELoaderType::Chunks };
+	FStreamHandler StreamHandlers[3];		// 0 = video (or multiplexed), 1 = audio, 2 = subtitle
+	FErrorDetail ErrorDetail;
+	IPlayerSessionServices* PlayerSessionService = nullptr;
+	bool bIsStarted = false;
+};
+
+
+
+} // namespace Electra
diff --git a/Engine/Plugins/Media/ElectraPlayer/Source/ElectraPlayerRuntime/Private/Runtime/PlayerFacility.h b/Engine/Plugins/Media/ElectraPlayer/Source/ElectraPlayerRuntime/Private/Runtime/PlayerFacility.h
--- a/Engine/Plugins/Media/ElectraPlayer/Source/ElectraPlayerRuntime/Private/Runtime/PlayerFacility.h
+++ b/Engine/Plugins/Media/ElectraPlayer/Source/ElectraPlayerRuntime/Private/Runtime/PlayerFacility.h
@@ -29,6 +29,9 @@ namespace Electra
 			MP4PlaylistReader,
 			MP4Playlist,
 			MP4StreamReader,
+			ISOBMFFPlaylistReader,
+			ISOBMFFPlaylist,
+			ISOBMFFStreamReader,
 			MKVPlaylistReader,
 			MKVPlaylist,
 			MKVStreamReader,
@@ -90,6 +93,12 @@ namespace Electra
 					return TEXT("MP4 playlist");
 				case EFacility::MP4StreamReader:
 					return TEXT("MP4 reader");
+				case EFacility::ISOBMFFPlaylistReader:
+					return TEXT("ISOBMFF playlist reader");
+				case EFacility::ISOBMFFPlaylist:
+					return TEXT("ISOBMFF playlist");
+				case EFacility::ISOBMFFStreamReader:
+					return TEXT("ISOBMFF reader");
 				case EFacility::LicenseKey:
 					return TEXT("License key");
 				case EFacility::DASHMPDReader:
diff --git a/Engine/Plugins/Media/ElectraPlayer/Source/ElectraPlayerRuntime/Private/Runtime/StreamAccessUnitBuffer.cpp b/Engine/Plugins/Media/ElectraPlayer/Source/ElectraPlayerRuntime/Private/Runtime/StreamAccessUnitBuffer.cpp
--- a/Engine/Plugins/Media/ElectraPlayer/Source/ElectraPlayerRuntime/Private/Runtime/StreamAccessUnitBuffer.cpp
+++ b/Engine/Plugins/Media/ElectraPlayer/Source/ElectraPlayerRuntime/Private/Runtime/StreamAccessUnitBuffer.cpp
@@ -300,6 +300,19 @@ namespace Electra
 		}
 	}
 
+	TSharedPtrTS<FAccessUnitBuffer> FMultiTrackAccessUnitBuffer::FindBufferFor(const TSharedPtrTS<const FBufferSourceInfo>& InBufferInfo)
+	{
+		FScopeLock lock(&AccessLock);
+		for(int32 i=0,iMax=BufferList.Num(); i<iMax; ++i)
+		{
+			if (BufferList[i].Info->PeriodAdaptationSetID.Equals(InBufferInfo->PeriodAdaptationSetID))
+			{
+				return BufferList[i].Buffer;
+			}
+		}
+		return nullptr;
+	}
+
 	bool FMultiTrackAccessUnitBuffer::PeekAndAddRef(FAccessUnit*& OutAU)
 	{
 		FScopedLock lock(AsShared());
diff --git a/Engine/Plugins/Media/ElectraPlayer/Source/ElectraPlayerRuntime/Private/Runtime/StreamAccessUnitBuffer.h b/Engine/Plugins/Media/ElectraPlayer/Source/ElectraPlayerRuntime/Private/Runtime/StreamAccessUnitBuffer.h
--- a/Engine/Plugins/Media/ElectraPlayer/Source/ElectraPlayerRuntime/Private/Runtime/StreamAccessUnitBuffer.h
+++ b/Engine/Plugins/Media/ElectraPlayer/Source/ElectraPlayerRuntime/Private/Runtime/StreamAccessUnitBuffer.h
@@ -593,7 +593,7 @@ namespace Electra
 			else
 			{
 				// Max allowed duration ok?
-				return AU->Duration + ExternalInfo->Duration > Limit->MaxDuration ? false : true;
+				return (AU->Duration + ExternalInfo->Duration) > Limit->MaxDuration ? false : true;
 			}
 		}
 
@@ -657,6 +657,8 @@ namespace Electra
 			TSharedPtrTS<FMultiTrackAccessUnitBuffer> LockedSelf;
 		};
 
+		TSharedPtrTS<FAccessUnitBuffer> FindBufferFor(const TSharedPtrTS<const FBufferSourceInfo>& InBufferInfo);
+
 		bool PeekAndAddRef(FAccessUnit*& OutAU);
 		bool Pop(FAccessUnit*& OutAU);
 		void PopDiscardUntil(FTimeValue UntilTime);
diff --git a/Engine/Plugins/Media/ElectraPlayer/Source/ElectraPlayerRuntime/Private/Runtime/StreamDataBuffer.h b/Engine/Plugins/Media/ElectraPlayer/Source/ElectraPlayerRuntime/Private/Runtime/StreamDataBuffer.h
--- a/Engine/Plugins/Media/ElectraPlayer/Source/ElectraPlayerRuntime/Private/Runtime/StreamDataBuffer.h
+++ b/Engine/Plugins/Media/ElectraPlayer/Source/ElectraPlayerRuntime/Private/Runtime/StreamDataBuffer.h
@@ -22,12 +22,12 @@ namespace Electra
 			Deallocate();
 		}
 
-		void SetExternalBuffer(uint8* InExternalBuffer, uint64 InExternalBufferSize)
+		void SetExternalBuffer(uint8* InExternalBuffer, int64 InExternalBufferSize)
 		{
 			FScopeLock Lock(&AccessLock);
 			Deallocate();
 			ExternalBuffer = InExternalBuffer;
-			DataSize = InExternalBufferSize;
+			DataSize = (uint64) InExternalBufferSize;
 		}
 
 		// Allocates buffer of the specified capacity, destroying any previous buffer.
diff --git a/Engine/Plugins/Media/ElectraPlayer/Source/ElectraPlayerRuntime/Private/Runtime/StreamTypes.cpp b/Engine/Plugins/Media/ElectraPlayer/Source/ElectraPlayerRuntime/Private/Runtime/StreamTypes.cpp
--- a/Engine/Plugins/Media/ElectraPlayer/Source/ElectraPlayerRuntime/Private/Runtime/StreamTypes.cpp
+++ b/Engine/Plugins/Media/ElectraPlayer/Source/ElectraPlayerRuntime/Private/Runtime/StreamTypes.cpp
@@ -9,7 +9,6 @@
 
 namespace Electra
 {
-
 	namespace
 	{
 		FString Printable4CC(const uint32 In4CC)
@@ -34,6 +33,128 @@ namespace Electra
 	}
 
 
+
+	FStreamCodecInformation::FStreamCodecInformation(const FCodecTypeFormat& InFromCodecTypeFormat)
+	{
+		Clear();
+		auto SetCommon = [&]() -> void
+		{
+			SetMimeType(InFromCodecTypeFormat.MimeType);
+			SetCodec4CC(InFromCodecTypeFormat.FourCC);
+			SetCodecSpecifierRFC6381(InFromCodecTypeFormat.RFC6381);
+			SetBitrate(InFromCodecTypeFormat.Bitrate);
+			SetCodecSpecificData(InFromCodecTypeFormat.CSD);
+			SetDecoderConfigRecord(InFromCodecTypeFormat.DCR);
+			SetStreamLanguageTag(InFromCodecTypeFormat.LanguageTag);
+			for(auto &It : InFromCodecTypeFormat.ExtraBoxes)
+			{
+				FString BoxName = Printable4CC(It.Key);
+				BoxName.Append(TEXT("_box"));
+				GetExtras().Set(FName(BoxName), FVariantValue(It.Value));
+			}
+		};
+		switch(InFromCodecTypeFormat.Type)
+		{
+			case FCodecTypeFormat::EType::Video:
+			{
+				const FCodecTypeFormat::FVideo& cvid = InFromCodecTypeFormat.Properties.Get<FCodecTypeFormat::FVideo>();
+
+				SetStreamType(EStreamType::Video);
+				if (InFromCodecTypeFormat.FourCC == ElectraDecodersUtil::Make4CC('a','v','c','1') || InFromCodecTypeFormat.FourCC == ElectraDecodersUtil::Make4CC('a','v','c','3'))
+				{
+					SetCodec(ECodec::H264);
+				}
+				else if (InFromCodecTypeFormat.FourCC == ElectraDecodersUtil::Make4CC('h','v','c','1') || InFromCodecTypeFormat.FourCC == ElectraDecodersUtil::Make4CC('h','e','v','1'))
+				{
+					SetCodec(ECodec::H265);
+				}
+				else
+				{
+					SetCodec(ECodec::Video4CC);
+				}
+				SetCommon();
+
+				SetResolution(FResolution((int32)cvid.Width, (int32)cvid.Height));
+				SetFrameRate(FTimeFraction(cvid.FrameRate.Numerator, cvid.FrameRate.Denominator));
+				SetAspectRatio(FAspectRatio(1, 1));
+
+				SetProfileTier((int32) cvid.Profile.Tier);
+				SetProfileSpace((int32) cvid.Profile.ProfileSpace);
+				SetProfile((int32) cvid.Profile.Profile);
+				SetProfileLevel((int32) cvid.Profile.Level);
+				SetProfileConstraints(cvid.Profile.Constraints);
+				SetProfileCompatibilityFlags((uint32) cvid.Profile.CompatibilityFlags);
+
+				SetHumanReadableCodecName(InFromCodecTypeFormat.HumanReadableFormatInfo);
+				break;
+			}
+			case FCodecTypeFormat::EType::Audio:
+			{
+				const FCodecTypeFormat::FAudio& caud = InFromCodecTypeFormat.Properties.Get<FCodecTypeFormat::FAudio>();
+
+				SetStreamType(EStreamType::Audio);
+				if (InFromCodecTypeFormat.FourCC == ElectraDecodersUtil::Make4CC('m','p','4','a'))
+				{
+					SetCodec(ECodec::AAC);
+				}
+				else if (InFromCodecTypeFormat.FourCC == ElectraDecodersUtil::Make4CC('e','c','-','3'))
+				{
+					SetCodec(ECodec::EAC3);
+				}
+				else if (InFromCodecTypeFormat.FourCC == ElectraDecodersUtil::Make4CC('a','c','-','3'))
+				{
+					SetCodec(ECodec::AC3);
+				}
+				else
+				{
+					SetCodec(ECodec::Audio4CC);
+				}
+				SetCommon();
+
+				SetSamplingRate((int32) caud.SampleRate);
+				SetNumberOfChannels((int32) caud.NumChannels);
+				SetChannelConfiguration(caud.ChannelConfiguration);
+
+				SetHumanReadableCodecName(InFromCodecTypeFormat.HumanReadableFormatInfo);
+				break;
+			}
+			case FCodecTypeFormat::EType::Subtitle:
+			{
+				const FCodecTypeFormat::FSubtitle& csub = InFromCodecTypeFormat.Properties.Get<FCodecTypeFormat::FSubtitle>();
+
+				SetStreamType(EStreamType::Subtitle);
+				if (InFromCodecTypeFormat.FourCC == ElectraDecodersUtil::Make4CC('w','v','t','t'))
+				{
+					SetCodec(ECodec::WebVTT);
+				}
+				else if (InFromCodecTypeFormat.FourCC == ElectraDecodersUtil::Make4CC('s','t','p','p') || InFromCodecTypeFormat.FourCC == ElectraDecodersUtil::Make4CC('t','t','m','l'))
+				{
+					SetCodec(ECodec::TTML);
+				}
+				else if (InFromCodecTypeFormat.FourCC == ElectraDecodersUtil::Make4CC('t','x','3','g'))
+				{
+					SetCodec(ECodec::TX3G);
+				}
+				else
+				{
+					SetCodec(ECodec::OtherSubtitle);
+				}
+				SetCommon();
+
+				SetHumanReadableCodecName(InFromCodecTypeFormat.HumanReadableFormatInfo);
+				break;
+			}
+			default:
+			{
+				break;
+			}
+		}
+	}
+
+
+
+
+
 	FString FStreamCodecInformation::GetMimeType() const
 	{
 		if (!MimeType.IsEmpty())
diff --git a/Engine/Plugins/Media/ElectraPlayer/Source/ElectraPlayerRuntime/Private/Runtime/StreamTypes.h b/Engine/Plugins/Media/ElectraPlayer/Source/ElectraPlayerRuntime/Private/Runtime/StreamTypes.h
--- a/Engine/Plugins/Media/ElectraPlayer/Source/ElectraPlayerRuntime/Private/Runtime/StreamTypes.h
+++ b/Engine/Plugins/Media/ElectraPlayer/Source/ElectraPlayerRuntime/Private/Runtime/StreamTypes.h
@@ -8,6 +8,7 @@
 #include "MediaVideoDecoderOutput.h"
 #include "MediaURLType.h"
 #include "Utilities/BCP47-Helpers.h"
+#include "CodecTypeFormat.h"
 
 namespace Electra
 {
@@ -63,6 +64,8 @@ namespace Electra
 		{
 			Clear();
 		}
+		explicit FStreamCodecInformation(const FCodecTypeFormat& InFromCodecTypeFormat);
+
 		enum class ECodec
 		{
 			// --- Unknown ---
@@ -756,6 +759,7 @@ namespace Electra
 		FString ID;
 		FString Kind;
 		FString Label;
+		TOptional<FTimespan> Duration;
 
 		TArray<FStreamMetadata> StreamDetails;
 		FStreamCodecInformation HighestBandwidthCodec;
diff --git a/Engine/Plugins/Media/ElectraPlayer/Source/ElectraPlayerRuntime/Private/Runtime/Utilities/MP4Helpers.cpp b/Engine/Plugins/Media/ElectraPlayer/Source/ElectraPlayerRuntime/Private/Runtime/Utilities/MP4Helpers.cpp
new file mode 100644
--- /dev/null
+++ b/Engine/Plugins/Media/ElectraPlayer/Source/ElectraPlayerRuntime/Private/Runtime/Utilities/MP4Helpers.cpp
@@ -0,0 +1,571 @@
+// Copyright Epic Games, Inc. All Rights Reserved.
+
+#include "MP4Helpers.h"
+#include "MP4Utilities.h"
+#include "Utilities/URLParser.h"
+#include "Utilities/BCP47-Helpers.h"
+#include "Features/IModularFeatures.h"
+#include "IElectraCodecFactoryModule.h"
+#include "IElectraCodecFactory.h"
+#include "IElectraCodecFormatUtils.h"
+#include "IElectraDecoderFeaturesAndOptions.h"
+#include "ElectraDecodersUtils.h"
+
+namespace Electra
+{
+
+FMP4DataLoader::~FMP4DataLoader()
+{
+	Close();
+}
+
+void FMP4DataLoader::Open(const FString& InURL, const FHttpParams& InHttpParams)
+{
+	URL = InURL;
+	// Check if this is a local file.
+	if (URL.StartsWith(TEXT("file://")))
+	{
+		FString Filename;
+		FURL_RFC3986::UrlDecode(Filename, InURL);
+		Filename.RightChopInline(7);
+		ConnectionInfo.bIsConnected = true;
+		ConnectionInfo.bHaveResponseHeaders = true;
+		ConnectionInfo.ContentType = TEXT("application/octet-stream");
+		ConnectionInfo.EffectiveURL = InURL;
+		ConnectionInfo.HTTPVersionReceived = 11;
+		ConnectionInfo.bIsChunked = false;
+
+		Archive = IFileManager::Get().CreateFileReader(*Filename, 0);
+		if (Archive)
+		{
+			FileSize = Archive->TotalSize();
+			ConnectionInfo.ContentLength = FileSize;
+			ConnectionInfo.StatusInfo.HTTPStatus = 200;
+			ConnectionInfo.ContentLengthHeader = FString::Printf(TEXT("Content-Length: %lld"), (long long int)FileSize);
+		}
+		else
+		{
+			bHasErrored = true;
+			LastError = FString::Printf(TEXT("Failed to open file \"%s\""), *URL);
+			ConnectionInfo.StatusInfo.HTTPStatus = 404;	// File not found
+			ConnectionInfo.StatusInfo.ErrorDetail.SetMessage(FString::Printf(TEXT("HTTP returned status 404")));
+			ConnectionInfo.StatusInfo.ErrorCode = HTTP::EStatusErrorCode::ERRCODE_HTTP_RETURNED_ERROR;
+		}
+	}
+	else
+	{
+		HttpParams = InHttpParams;
+		ProgressListener = MakeShared<IElectraHttpManager::FProgressListener, ESPMode::ThreadSafe>();
+		ProgressListener->CompletionDelegate = IElectraHttpManager::FProgressListener::FCompletionDelegate::CreateThreadSafeSP(AsShared(), &FMP4DataLoader::HTTPCompletionHandle);
+	}
+}
+
+void FMP4DataLoader::Close()
+{
+	if (Archive)
+	{
+		Archive->Close();
+		delete Archive;
+		Archive = nullptr;
+	}
+	HttpParams.HTTPManager.Reset();
+	HttpParams.HTTPResponseCache.Reset();
+	ProgressListener.Reset();
+}
+
+void FMP4DataLoader::HTTPCompletionHandle(const IElectraHttpManager::FRequest* InRequest)
+{
+	bool bFailed = InRequest->ConnectionInfo.StatusInfo.ErrorDetail.IsError();
+	ConnectionInfo = InRequest->ConnectionInfo;
+	if (!bFailed)
+	{
+		// Set the size of the resource if we don't have it yet.
+		if (FileSize < 0)
+		{
+			ElectraHTTPStream::FHttpRange crh;
+			if (crh.ParseFromContentRangeResponse(InRequest->ConnectionInfo.ContentRangeHeader))
+			{
+				FileSize = crh.GetDocumentSize();
+			}
+		}
+		if (ConnectionInfo.EffectiveURL.Len())
+		{
+			URL = ConnectionInfo.EffectiveURL;
+		}
+	}
+	bHasErrored = bFailed;
+	ReadCompleted.Signal();
+}
+
+
+int64 FMP4DataLoader::ReadData(void* InOutBuffer, int64 InNumBytes, int64 InFromOffset, FCancellationCheckDelegate InCheckCancellationDelegate)
+{
+	if (bHasErrored)
+	{
+		return EResult::ReadError;
+	}
+	if (InNumBytes <= 0)
+	{
+		return 0;
+	}
+	if (Archive)
+	{
+		if (CurrentOffset != InFromOffset)
+		{
+			CurrentOffset = InFromOffset >= 0 ? InFromOffset <= FileSize ? InFromOffset : FileSize : 0;
+			Archive->Seek(CurrentOffset);
+		}
+		if (InNumBytes + CurrentOffset > FileSize)
+		{
+			InNumBytes = FileSize - CurrentOffset;
+		}
+		if (InOutBuffer)
+		{
+			Archive->Serialize(InOutBuffer, InNumBytes);
+		}
+		else
+		{
+			Archive->Seek(CurrentOffset + InNumBytes);
+		}
+		CurrentOffset += InNumBytes;
+		return InNumBytes;
+	}
+	else
+	{
+		// Skip data?
+		if (!InOutBuffer)
+		{
+			CurrentOffset = InFromOffset + InNumBytes;
+			return InNumBytes;
+		}
+
+		ProgressListener->ProgressDelegate = IElectraHttpManager::FProgressListener::FProgressDelegate::CreateLambda([&](const IElectraHttpManager::FRequest* InRequest)->int32
+		{
+			bWasAborted = InCheckCancellationDelegate.Execute();
+			return bWasAborted ? 1 : 0;
+		});
+
+
+		constexpr int32 kSmallReadChunkSize = 4096;
+		// Loop the read attempt in case the first request fails with a "416 - Range Not Satisfiable"
+		bool bForceFirstReadAsSmall = false;
+		while(1)
+		{
+			TSharedPtr<FWaitableBuffer, ESPMode::ThreadSafe> TargetBuffer = MakeShared<FWaitableBuffer, ESPMode::ThreadSafe>();
+			bool bIsSmallRead = !bForceFirstReadAsSmall && InNumBytes < kSmallReadChunkSize;
+			bool bReadViaTempBuffer = false;
+			int64 NumBytesToRead = InNumBytes;
+			if (bIsSmallRead)
+			{
+				// Did we read past what we have to read now from an earlier read?
+				if (InFromOffset + InNumBytes <= CurrentEndOffsetIncluding)
+				{
+					TargetBuffer->SetExternalBuffer(static_cast<uint8*>(InOutBuffer), NumBytesToRead);
+				}
+				else
+				{
+					NumBytesToRead = kSmallReadChunkSize;
+					TargetBuffer->Reserve(NumBytesToRead);
+					bReadViaTempBuffer = true;
+				}
+			}
+			else
+			{
+				TargetBuffer->SetExternalBuffer(static_cast<uint8*>(InOutBuffer), NumBytesToRead);
+			}
+
+			TSharedPtr<IElectraHttpManager::FRequest, ESPMode::ThreadSafe> Request = MakeShared<IElectraHttpManager::FRequest, ESPMode::ThreadSafe>();
+			Request->Parameters.URL = URL;
+			Request->Parameters.Range.SetStart(InFromOffset);
+			int64 LastByte = InFromOffset + NumBytesToRead - 1;
+			if (FileSize >= 0 && LastByte > FileSize-1)
+			{
+				LastByte = FileSize - 1;
+			}
+			CurrentEndOffsetIncluding = LastByte > CurrentEndOffsetIncluding ? LastByte : CurrentEndOffsetIncluding;
+			Request->Parameters.Range.SetEndIncluding(LastByte);
+			Request->Parameters.ConnectTimeout = FTimeValue().SetFromMilliseconds(HttpParams.ConnectTimeoutMillis);
+			Request->Parameters.NoDataTimeout = FTimeValue().SetFromMilliseconds(HttpParams.NoDataTimeoutMillis);
+			Request->Parameters.AcceptEncoding.Set(TEXT("identity"));
+			Request->Parameters.RequestHeaders = HttpParams.RequestHeaders;
+			Request->ProgressListener = ProgressListener;
+			Request->ResponseCache = HttpParams.HTTPResponseCache;
+			Request->ReceiveBuffer = TargetBuffer;
+
+			HttpParams.HTTPManager->AddRequest(Request, false);
+			ReadCompleted.WaitAndReset();
+			HttpParams.HTTPManager->RemoveRequest(Request, false);
+			LastError = Request->ConnectionInfo.StatusInfo.ErrorDetail.GetMessage();
+			Request.Reset();
+
+			// Check if this is the first attempt for which we may have requested additional data to be read
+			// and the server returned an error because of that.
+			if (bIsFirstRead && bHasErrored && ConnectionInfo.StatusInfo.HTTPStatus == 416)
+			{
+				bIsFirstRead = false;
+				bForceFirstReadAsSmall = true;
+				CurrentEndOffsetIncluding = -1;
+				bHasErrored = false;
+				ConnectionInfo = HTTP::FConnectionInfo();
+				continue;
+			}
+
+			if (bHasErrored)
+			{
+				return EResult::ReadError;
+			}
+			else if (bWasAborted)
+			{
+				return EResult::Canceled;
+			}
+
+			if (TargetBuffer->Num() < InNumBytes)
+			{
+				LastError = FString::Printf(TEXT("Got only %d bytes instead of %lld"), TargetBuffer->Num(), (long long int)InNumBytes);
+				return EResult::ReadError;
+			}
+			if (bReadViaTempBuffer)
+			{
+				FMemory::Memcpy(InOutBuffer, TargetBuffer->GetLinearReadData(), InNumBytes);
+			}
+			break;
+		}
+		CurrentOffset = InFromOffset + InNumBytes;
+		return InNumBytes;
+	}
+}
+
+bool FMP4DataLoader::HasReachedEOF()
+{
+	bool bEOS = FileSize >= 0 && CurrentOffset >= FileSize;
+	return bEOS;
+}
+
+
+
+
+bool GetTrackFormatInfo(FString& OutMessage, FCodecTypeFormat& OutCodecInfo, FDRMTypeFormat& OutDRMInfo, FDecoderInformation& OutDecoderInfo, const TSharedPtr<MP4Boxes::FMP4BoxTRAK, ESPMode::ThreadSafe>& InTrack, uint32 InTrackID, bool bInCheckIfDecodable)
+{
+	IElectraCodecFormatUtilsModularFeature* CodecFormatUtils = nullptr;
+	{
+		IModularFeatures::FScopedLockModularFeatureList ScopedLockModularFeatureList;
+		TArray<IElectraCodecFormatUtilsModularFeature*> FormatUtils = IModularFeatures::Get().GetModularFeatureImplementations<IElectraCodecFormatUtilsModularFeature>(IElectraCodecFormatUtilsModularFeature::GetModularFeatureName());
+		CodecFormatUtils = !FormatUtils.IsEmpty() ? FormatUtils[0] : nullptr;
+	}
+
+	if (!CodecFormatUtils)
+	{
+		OutMessage = TEXT("No registered IElectraCodecFormatUtilsModularFeature found. Unable to determine codecs, ignoring track.");
+		return false;
+	}
+
+	IElectraCodecFactoryModule* CodecFactoryModule = nullptr;
+	if (bInCheckIfDecodable)
+	{
+		CodecFactoryModule = static_cast<IElectraCodecFactoryModule*>(FModuleManager::Get().GetModule(TEXT("ElectraCodecFactory")));
+		if (!CodecFactoryModule)
+		{
+			OutMessage = FString::Printf(TEXT("Electra decoder factory not found. Unable to use any track."));
+			return false;
+		}
+	}
+	auto CheckIfDecodable = [&]() -> void
+	{
+		// Shall we determine if there is a decoder able to decode this track?
+		if (CodecFactoryModule)
+		{
+			TMap<FString, FVariant> FormatInfo;
+			auto Factory = CodecFactoryModule->GetBestDecoderFactoryForFormat(FormatInfo, OutCodecInfo, TMap<FString, FVariant>());
+			if (Factory.IsValid())
+			{
+				OutDecoderInfo.bIsDecodable = true;
+				// Is there missing information the decoder factory could provide?
+				TMap<FString, FVariant> ConfigOptions;
+				Factory->GetConfigurationOptions(ConfigOptions);
+				if (OutCodecInfo.KeyframeMode == FCodecTypeFormat::EKeyframeMode::Unknown && ConfigOptions.Find(IElectraDecoderFormatInfo::IsEveryFrameKeyframe))
+				{
+					OutCodecInfo.KeyframeMode = !!ElectraDecodersUtil::GetVariantValueSafeI64(FormatInfo, IElectraDecoderFormatInfo::IsEveryFrameKeyframe, 0)
+												? FCodecTypeFormat::EKeyframeMode::OnlyKeyframes : FCodecTypeFormat::EKeyframeMode::DeltaFrames;
+				}
+				if (OutCodecInfo.HumanReadableFormatInfo.IsEmpty() && ConfigOptions.Find(IElectraDecoderFormatInfo::HumanReadableFormatName))
+				{
+					OutCodecInfo.HumanReadableFormatInfo = ElectraDecodersUtil::GetVariantValueFString(FormatInfo, IElectraDecoderFormatInfo::HumanReadableFormatName);
+				}
+				OutDecoderInfo.ProviderInfo.Name = Factory->GetProviderInformation().GetName();
+				OutDecoderInfo.ProviderInfo.Version = Factory->GetProviderInformation().GetVersion();
+				OutDecoderInfo.ProviderInfo.Implementation = Factory->GetProviderInformation().GetImplementation();
+				OutDecoderInfo.ProviderInfo.Vendor = Factory->GetProviderInformation().GetVendor();
+			}
+			else
+			{
+				OutDecoderInfo.bIsDecodable = false;
+			}
+		}
+	};
+
+	// There needs to be an `stsd` box in this track. We do not try the expected path of `trak`->`mdia`->`minf`->`stbl`->`stsd` as
+	// this is not that much faster and if the file is somewhat ill-formed we may not find it if it's grouped under elsewhere.
+	auto Stsd = InTrack->FindBoxRecursive<MP4Boxes::FMP4BoxSTSD>(MP4Utilities::MakeBoxAtom('s','t','s','d'), 6);
+	if (!Stsd)
+	{
+		// If not found we just ignore the track. That's a warning but not an error.
+		OutMessage = FString::Printf(TEXT("No `stsd` box found, ignoring track #%u."), InTrackID);
+		return false;
+	}
+
+	// Get language code
+	auto Mdhd = InTrack->FindBoxRecursive<MP4Boxes::FMP4BoxMDHD>(MP4Utilities::MakeBoxAtom('m','d','h','d'), 2);
+	if (Mdhd)
+	{
+		Electra::BCP47::ParseRFC5646Tag(OutCodecInfo.LanguageTag, Mdhd->GetLanguageCode639_2T());
+	}
+	auto Elng = InTrack->FindBoxRecursive<MP4Boxes::FMP4BoxELNG>(MP4Utilities::MakeBoxAtom('e','l','n','g'), 6);
+	if (Elng)
+	{
+		Electra::BCP47::ParseRFC5646Tag(OutCodecInfo.LanguageTag, Elng->GetLanguageTag());
+	}
+
+	auto AddChildren = [&OutCodecInfo](const TSharedPtr<MP4Boxes::FMP4BoxBase, ESPMode::ThreadSafe>& InSample) -> void
+	{
+		for(auto& Ch : InSample->GetChildren())
+		{
+			OutCodecInfo.ExtraBoxes.Emplace(Ch->GetType(), Ch->GetBoxData());
+		}
+	};
+
+	// Get the sample type of this track. This call is necessary to trigger parsing of child nodes.
+	MP4Boxes::FMP4BoxSampleEntry::ESampleType SampleType = Stsd->GetSampleType();
+	// If already knows to not be supported, skip it.
+	if (SampleType == MP4Boxes::FMP4BoxSampleEntry::ESampleType::Unsupported)
+	{
+		OutMessage = FString::Printf(TEXT("Unsupported sample type, ignoring track #%u."), InTrackID);
+		return false;
+	}
+	// Several entries are permitted, but we need this to be unambiguous.
+	if (Stsd->GetChildren().Num() > 1)
+	{
+		OutMessage = FString::Printf(TEXT("Multiple sample descriptions found in `stsd` box, ignoring track #%u."), InTrackID);
+		return false;
+	}
+	else if (Stsd->GetChildren().Num() == 0)
+	{
+		OutMessage = FString::Printf(TEXT("No sample description found in `stsd` box, ignoring track #%u."), InTrackID);
+		return false;
+	}
+
+	OutDecoderInfo.bIsDecodable = true;
+	OutDRMInfo.bIsEncrypted = false;
+	// Based on the sample type, get it and see if it is using codec we support.
+	if (SampleType == MP4Boxes::FMP4BoxSampleEntry::ESampleType::Video)
+	{
+		TArray<TSharedPtr<MP4Boxes::FMP4BoxVisualSampleEntry, ESPMode::ThreadSafe>> Visuals;
+		Stsd->GetSampleDescriptions(Visuals);
+		if (Visuals.Num() != 1)
+		{
+			if (Visuals.IsEmpty())
+			{
+				OutMessage = FString::Printf(TEXT("No visual sample entry found in `stsd` box, ignoring track #%u."), InTrackID);
+			}
+			else
+			{
+				OutMessage = FString::Printf(TEXT("More than one visual sample entry found in `stsd` box, ignoring track #%u."), InTrackID);
+			}
+			return false;
+		}
+
+		switch(Visuals[0]->GetType())
+		{
+			case MP4Utilities::MakeBoxAtom('e','n','c','v'):
+			{
+				OutDRMInfo.bIsEncrypted = true;
+				break;
+			}
+			default:
+			{
+				if (Visuals[0]->GetFrameCount() == 1)
+				{
+					OutCodecInfo.Type = Electra::FCodecTypeFormat::EType::Video;
+					OutCodecInfo.FourCC = Visuals[0]->GetType();
+					OutCodecInfo.Properties.Emplace<Electra::FCodecTypeFormat::FVideo>();
+					Electra::FCodecTypeFormat::FVideo& cvid = OutCodecInfo.Properties.Get<Electra::FCodecTypeFormat::FVideo>();
+					cvid.Width = Visuals[0]->GetWidth();
+					cvid.Height = Visuals[0]->GetHeight();
+					auto Stts = InTrack->FindBoxRecursive<MP4Boxes::FMP4BoxSTTS>(MP4Utilities::MakeBoxAtom('s','t','t','s'), 5);
+					if (Mdhd.IsValid() && Stts.IsValid() && Stts->GetEntries().Num())
+					{
+						cvid.FrameRate = FFrameRate(Mdhd->GetTimescale(), Stts->GetEntries()[0].sample_delta);
+					}
+					for(auto& Ch : Visuals[0]->GetChildren())
+					{
+						OutCodecInfo.ExtraBoxes.Emplace(Ch->GetType(), Ch->GetBoxData());
+					}
+					if (!CodecFormatUtils->PrepareCodecTypeFormat(OutCodecInfo, IElectraCodecFormatUtilsModularFeature::FCodecTypeFormatParams()))
+					{
+						OutMessage = FString::Printf(TEXT("Failed to prepare the codec format type, ignoring track #%u."), InTrackID);
+						return false;
+					}
+					// Check if this video is decodable.
+					CheckIfDecodable();
+				}
+				else
+				{
+					OutMessage = FString::Printf(TEXT("Track #%u has a frame count other than 1 in the VisualSampleEntry, ignoring this track."), InTrackID);
+					return false;
+				}
+				break;
+			}
+		}
+	}
+	else if (SampleType == MP4Boxes::FMP4BoxSampleEntry::ESampleType::Audio)
+	{
+		TArray<TSharedPtr<MP4Boxes::FMP4BoxAudioSampleEntry, ESPMode::ThreadSafe>> Audios;
+		Stsd->GetSampleDescriptions(Audios);
+		if (Audios.Num() != 1)
+		{
+			if (Audios.IsEmpty())
+			{
+				OutMessage = FString::Printf(TEXT("No audio sample entry found in `stsd` box, ignoring track #%u."), InTrackID);
+			}
+			else
+			{
+				OutMessage = FString::Printf(TEXT("More than one audio sample entry found in `stsd` box, ignoring track #%u."), InTrackID);
+			}
+			return false;
+		}
+
+		switch(Audios[0]->GetType())
+		{
+			case MP4Utilities::MakeBoxAtom('e','n','c','a'):
+			{
+				OutDRMInfo.bIsEncrypted = true;
+				break;
+			}
+			default:
+			{
+				OutCodecInfo.Type = Electra::FCodecTypeFormat::EType::Audio;
+				OutCodecInfo.FourCC = Audios[0]->GetType();
+				OutCodecInfo.Properties.Emplace<Electra::FCodecTypeFormat::FAudio>();
+				Electra::FCodecTypeFormat::FAudio& caud = OutCodecInfo.Properties.Get<Electra::FCodecTypeFormat::FAudio>();
+				caud.NumChannels = Audios[0]->GetChannelCount();
+				caud.SampleRate = Audios[0]->GetSampleRate();
+				for(auto& Ch : Audios[0]->GetChildren())
+				{
+					OutCodecInfo.ExtraBoxes.Emplace(Ch->GetType(), Ch->GetBoxData());
+				}
+				if (!CodecFormatUtils->PrepareCodecTypeFormat(OutCodecInfo, IElectraCodecFormatUtilsModularFeature::FCodecTypeFormatParams()))
+				{
+					OutMessage = FString::Printf(TEXT("Failed to prepare the codec format type, ignoring track #%u."), InTrackID);
+					return false;
+				}
+				// Check if this video is decodable.
+				CheckIfDecodable();
+				break;
+			}
+		}
+
+	}
+	else if (SampleType == MP4Boxes::FMP4BoxSampleEntry::ESampleType::Subtitles || SampleType == MP4Boxes::FMP4BoxSampleEntry::ESampleType::TimedText)
+	{
+		TArray<TSharedPtr<MP4Boxes::FMP4BoxSampleEntry, ESPMode::ThreadSafe>> Subtitles;
+		Stsd->GetSampleDescriptions(Subtitles);
+		if (Subtitles.Num() != 1)
+		{
+			if (Subtitles.IsEmpty())
+			{
+				OutMessage = FString::Printf(TEXT("No subtitle sample entry found in `stsd` box, ignoring track #%u."), InTrackID);
+			}
+			else
+			{
+				OutMessage = FString::Printf(TEXT("More than one subtitle sample entry found in `stsd` box, ignoring track #%u."), InTrackID);
+			}
+			return false;
+		}
+
+		switch(Subtitles[0]->GetType())
+		{
+			case MP4Utilities::MakeBoxAtom('e','n','c','t'):
+			case MP4Utilities::MakeBoxAtom('e','n','c','u'):
+			{
+				OutDRMInfo.bIsEncrypted = true;
+				break;
+			}
+			default:
+			{
+				OutCodecInfo.Type = Electra::FCodecTypeFormat::EType::Subtitle;
+				OutCodecInfo.FourCC = Subtitles[0]->GetType();
+				OutCodecInfo.Properties.Emplace<Electra::FCodecTypeFormat::FSubtitle>();
+				//Electra::FCodecTypeFormat::FSubtitle& csub = OutCodecInfo.Properties.Get<Electra::FCodecTypeFormat::FSubtitle>();
+				// Add the sample box itself as an extra box.
+				OutCodecInfo.ExtraBoxes.Emplace(Subtitles[0]->GetType(), Subtitles[0]->GetBoxData());
+				for(auto& Ch : Subtitles[0]->GetChildren())
+				{
+					OutCodecInfo.ExtraBoxes.Emplace(Ch->GetType(), Ch->GetBoxData());
+				}
+				if (!CodecFormatUtils->PrepareCodecTypeFormat(OutCodecInfo, IElectraCodecFormatUtilsModularFeature::FCodecTypeFormatParams()))
+				{
+					OutMessage = FString::Printf( TEXT("Failed to prepare the codec format type, ignoring track #%u."), InTrackID);
+					return false;
+				}
+				break;
+			}
+		}
+	}
+	else if (SampleType == MP4Boxes::FMP4BoxSampleEntry::ESampleType::QTFFTimecode)
+	{
+		TArray<TSharedPtr<MP4Boxes::FMP4BoxQTFFTimecodeSampleEntry, ESPMode::ThreadSafe>> Timecodes;
+		Stsd->GetSampleDescriptions(Timecodes);
+		if (Timecodes.Num() != 1)
+		{
+			if (Timecodes.IsEmpty())
+			{
+				OutMessage = FString::Printf(TEXT("No timecode sample entry found in `stsd` box, ignoring track #%u."), InTrackID);
+			}
+			else
+			{
+				OutMessage = FString::Printf(TEXT("More than one timecode sample entry found in `stsd` box, ignoring track #%u."), InTrackID);
+			}
+			return false;
+		}
+
+		switch(Timecodes[0]->GetType())
+		{
+			case MP4Utilities::MakeBoxAtom('t','m','c','d'):
+			{
+				Electra::FCodecTypeFormat::FTMCDTimecode Timecode;
+
+				Timecode.Flags = Timecodes[0]->GetFlags();
+				Timecode.Timescale = Timecodes[0]->GetTimescale();
+				Timecode.FrameDuration = Timecodes[0]->GetFrameDuration();
+				Timecode.NumberOfFrames = Timecodes[0]->GetNumberOfFrames();
+				OutCodecInfo.Properties.Emplace<Electra::FCodecTypeFormat::FTMCDTimecode>(MoveTemp(Timecode));
+				OutCodecInfo.RFC6381 = MP4Utilities::GetPrintableBoxAtom(Timecodes[0]->GetType());
+				OutCodecInfo.FourCC = Timecodes[0]->GetType();
+				OutCodecInfo.DCR = Timecodes[0]->GetBoxData();
+				OutCodecInfo.Type = Electra::FCodecTypeFormat::EType::Timecode;
+				AddChildren(Timecodes[0]);
+				break;
+			}
+		}
+	}
+
+	// If the track has not been identified as usable so far, remove it.
+	if (OutCodecInfo.Type == Electra::FCodecTypeFormat::EType::Invalid)
+	{
+		if (!OutDRMInfo.bIsEncrypted)
+		{
+			TArray<TSharedPtr<MP4Boxes::FMP4BoxSampleEntry, ESPMode::ThreadSafe>> Entries;
+			Stsd->GetSampleDescriptions(Entries);
+			check(Entries.Num() == 1);
+			OutMessage = FString::Printf(TEXT("Track of sample type \"%s\" is not supported, ignoring track #%u."), *MP4Utilities::GetPrintableBoxAtom(Entries[0]->GetType()), InTrackID);
+		}
+		else
+		{
+			OutMessage = FString::Printf(TEXT("Track is using encryption, ignoring track #%u."), InTrackID);
+		}
+		return false;
+	}
+	return true;
+}
+
+
+} // namespace Electra
diff --git a/Engine/Plugins/Media/ElectraPlayer/Source/ElectraPlayerRuntime/Private/Runtime/Utilities/MP4Helpers.h b/Engine/Plugins/Media/ElectraPlayer/Source/ElectraPlayerRuntime/Private/Runtime/Utilities/MP4Helpers.h
new file mode 100644
--- /dev/null
+++ b/Engine/Plugins/Media/ElectraPlayer/Source/ElectraPlayerRuntime/Private/Runtime/Utilities/MP4Helpers.h
@@ -0,0 +1,72 @@
+// Copyright Epic Games, Inc. All Rights Reserved.
+
+#pragma once
+
+#include "CoreTypes.h"
+#include "Templates/SharedPointer.h"
+#include "HAL/FileManager.h"
+#include "Serialization/Archive.h"
+
+#include "HTTP/HTTPManager.h"
+#include "Core/MediaEventSignal.h"
+#include "MP4DataReader.h"
+#include "MP4Boxes.h"
+#include "CodecTypeFormat.h"
+
+namespace Electra
+{
+	class FMP4DataLoader : public MP4Utilities::IMP4DataReaderBase, public TSharedFromThis<FMP4DataLoader, ESPMode::ThreadSafe>
+	{
+	public:
+		static TSharedPtr<FMP4DataLoader, ESPMode::ThreadSafe> Create()
+		{ return MakeShareable(new FMP4DataLoader); }
+		virtual ~FMP4DataLoader();
+
+		struct FHttpParams
+		{
+			TSharedPtr<IElectraHttpManager, ESPMode::ThreadSafe> HTTPManager;
+			TSharedPtr<IHTTPResponseCache, ESPMode::ThreadSafe> HTTPResponseCache;
+			TArray<HTTP::FHTTPHeader> RequestHeaders;
+			int32 ConnectTimeoutMillis = 8000;
+			int32 NoDataTimeoutMillis = 6000;
+		};
+		void Open(const FString& InURL, const FHttpParams& InHttpParams);
+		void Close();
+
+		//! Methods from MP4Utilities::IMP4DataReaderBase
+		int64 ReadData(void* InOutBuffer, int64 InNumBytes, int64 InFromOffset, FCancellationCheckDelegate InCheckCancellationDelegate) override;
+		int64 GetTotalFileSize() override
+		{ return FileSize; }
+		int64 GetCurrentFileOffset() override
+		{ return CurrentOffset; }
+		bool HasReachedEOF() override;
+		FString GetLastError() override
+		{ return LastError; }
+
+		const HTTP::FConnectionInfo& GetConnectionInfo() const
+		{ return ConnectionInfo; }
+
+	private:
+		FMP4DataLoader() = default;
+		void HTTPCompletionHandle(const IElectraHttpManager::FRequest* InRequest);
+		FHttpParams HttpParams;
+		HTTP::FConnectionInfo ConnectionInfo;
+		FString URL;
+		FString LastError;
+		FMediaEvent ReadCompleted;
+		TSharedPtr<IElectraHttpManager::FProgressListener, ESPMode::ThreadSafe> ProgressListener;
+		FArchive* Archive = nullptr;
+		int64 FileSize = -1;
+		int64 CurrentOffset = 0;
+		int64 CurrentEndOffsetIncluding = -1;
+		bool bIsFirstRead = true;
+		volatile bool bHasErrored = false;
+		volatile bool bWasAborted = false;
+	};
+
+
+
+	bool GetTrackFormatInfo(FString& OutMessage, FCodecTypeFormat& OutCodecInfo, FDRMTypeFormat& OutDRMInfo, FDecoderInformation& OutDecoderInfo, const TSharedPtr<MP4Boxes::FMP4BoxTRAK, ESPMode::ThreadSafe>& InTrack, uint32 InTrackID, bool bInCheckIfDecodable);
+
+
+} // namespace Electra
diff --git a/Engine/Plugins/Media/ElectraPlayer/Source/ElectraProtron/Private/Player/ElectraProtronPlayerImpl.cpp b/Engine/Plugins/Media/ElectraPlayer/Source/ElectraProtron/Private/Player/ElectraProtronPlayerImpl.cpp
--- a/Engine/Plugins/Media/ElectraPlayer/Source/ElectraProtron/Private/Player/ElectraProtronPlayerImpl.cpp
+++ b/Engine/Plugins/Media/ElectraPlayer/Source/ElectraProtron/Private/Player/ElectraProtronPlayerImpl.cpp
@@ -15,6 +15,15 @@
 #include "Core/MediaThreads.h"
 #include "Features/IModularFeatures.h"
 
+/*********************************************************************************************************************/
+#include "HAL/IConsoleManager.h"
+static TAutoConsoleVariable<bool> CVarProtronForceShortestDuration(
+	TEXT("ElectraProtron.ForceShortestDuration"),
+	true,
+	TEXT("Set movie duration to the shortest track in the file.\n")
+	TEXT(" true: use the shortest duration\n false: keep individual track durations"),
+	ECVF_Default);
+/*********************************************************************************************************************/
 
 namespace ElectraProtronOptionNames
 {
@@ -510,12 +519,15 @@ void FElectraProtronPlayer::FImpl::InternalOpen(const FString& InFilename)
 	}
 	MP4Utilities::FFractionalTime EntireMovieDuration(MovieDuration);
 
-	// Check that the movie duration is not larger than the shortest supported track.
-	if (ShortestSupportedTrackDuration < TNumericLimits<int64>::Max() && MvhdDur > ShortestSupportedTrackDuration)
+	if (CVarProtronForceShortestDuration.GetValueOnAnyThread())
 	{
-		UE_LOG(LogElectraProtron, Warning, TEXT("Movie duration in `mvhd` box (%lld) is larger than the shortest supported track (%lld) in the file, adjusting movie duration down."), (long long int)MvhdDur, (long long int)ShortestSupportedTrackDuration);
-		MovieDuration.SetFromND(ShortestSupportedTrackDuration, MovieDuration.GetDenominator());
-		Duration = MovieDuration.GetAsTimespan();
+		// Check that the movie duration is not larger than the shortest supported track.
+		if (ShortestSupportedTrackDuration < TNumericLimits<int64>::Max() && MvhdDur > ShortestSupportedTrackDuration)
+		{
+			UE_LOG(LogElectraProtron, Warning, TEXT("Movie duration in `mvhd` box (%lld) is larger than the shortest supported track (%lld) in the file, adjusting movie duration down."), (long long int)MvhdDur, (long long int)ShortestSupportedTrackDuration);
+			MovieDuration.SetFromND(ShortestSupportedTrackDuration, MovieDuration.GetDenominator());
+			Duration = MovieDuration.GetAsTimespan();
+		}
 	}
 
 	// If there are timecode tracks, find which tracks references them.
@@ -540,13 +552,13 @@ void FElectraProtronPlayer::FImpl::InternalOpen(const FString& InFilename)
 					{
 						if (References.Num() > 1)
 						{
-							UE_LOG(LogElectraProtron, Warning, TEXT("Track %u contains more than one `tmcd` reference box. Using first reference only."), Tracks[TrkNum]->TrackID);
+							UE_LOG(LogElectraProtron, Warning, TEXT("Track #%u contains more than one `tmcd` reference box. Using first reference only."), Tracks[TrkNum]->TrackID);
 						}
 						if (References[0].TrackIDs.Num())
 						{
 							if (References[0].TrackIDs.Num() > 1)
 							{
-								UE_LOG(LogElectraProtron, Warning, TEXT("Track %u references more than one timecode track. Using first reference only."), Tracks[TrkNum]->TrackID);
+								UE_LOG(LogElectraProtron, Warning, TEXT("Track #%u references more than one timecode track. Using first reference only."), Tracks[TrkNum]->TrackID);
 							}
 							// Whether the reference is actually valid or not, a track makes an explicit reference
 							// and so we must not assign unreferenced timecode tracks later.
@@ -565,7 +577,7 @@ void FElectraProtronPlayer::FImpl::InternalOpen(const FString& InFilename)
 								}
 								else
 								{
-									UE_LOG(LogElectraProtron, Warning, TEXT("Track %u references a non-existing timecode track #%u. Ignoring."), Tracks[TrkNum]->TrackID, RefTkId);
+									UE_LOG(LogElectraProtron, Warning, TEXT("Track #%u references a non-existing timecode track #%u. Ignoring."), Tracks[TrkNum]->TrackID, RefTkId);
 									Tracks[TrkNum]->ReferencedTimecodeTrack.Reset();
 								}
 							}
@@ -684,6 +696,7 @@ void FElectraProtronPlayer::FImpl::InternalOpen(const FString& InFilename)
 			check(!LastErrorMessage.IsEmpty());
 			return;
 		}
+		Tracks[i]->Duration = Tracks[i]->MP4Track->GetMappedTrackDuration().GetAsTimespan();
 
 		// If this is a timecode track we may want to read in the first timecode.
 		if (Config.bReadFirstTimecode && Tracks[i]->CodecInfo.Type == Electra::FCodecTypeFormat::EType::Timecode &&
diff --git a/Engine/Plugins/Media/ElectraPlayer/Source/ElectraProtron/Private/Player/ElectraProtronPlayerImpl.h b/Engine/Plugins/Media/ElectraPlayer/Source/ElectraProtron/Private/Player/ElectraProtronPlayerImpl.h
--- a/Engine/Plugins/Media/ElectraPlayer/Source/ElectraProtron/Private/Player/ElectraProtronPlayerImpl.h
+++ b/Engine/Plugins/Media/ElectraPlayer/Source/ElectraProtron/Private/Player/ElectraProtronPlayerImpl.h
@@ -351,6 +351,7 @@ private:
 			FString Implementation;
 			FString Vendor;
 		};
+		FTimespan Duration;
 		FDecoderProviderInfo DecoderProviderInfo;
 		TSharedPtr<MP4Boxes::FMP4BoxTRAK, ESPMode::ThreadSafe> TrackBox;
 		TSharedPtr<MP4Boxes::FMP4Track::FFragmentInfo, ESPMode::ThreadSafe> FragmentInfo;
diff --git a/Engine/Plugins/Media/ElectraUtil/Source/ElectraBase/Public/CodecTypeFormat.h b/Engine/Plugins/Media/ElectraUtil/Source/ElectraBase/Public/CodecTypeFormat.h
--- a/Engine/Plugins/Media/ElectraUtil/Source/ElectraBase/Public/CodecTypeFormat.h
+++ b/Engine/Plugins/Media/ElectraUtil/Source/ElectraBase/Public/CodecTypeFormat.h
@@ -33,6 +33,12 @@ namespace Electra
 			Invalid,
 			MAX = Invalid
 		};
+		enum class EKeyframeMode
+		{
+			Unknown,
+			OnlyKeyframes,
+			DeltaFrames
+		};
 		struct FVideo
 		{
 			// Codec specific profile, level, etc.
@@ -125,11 +131,50 @@ namespace Electra
 		TArray<uint8> DCR;
 		TArray<uint8> CSD;
 		EType Type = EType::Invalid;
+		EKeyframeMode KeyframeMode = EKeyframeMode::Unknown;
 		uint32 FourCC = 0;
 		uint32 Bitrate = 0;
 		uint32 AverageBitrate = 0;
 	};
 
+
+	/**
+	 *
+	 */
+	struct FDRMTypeFormat
+	{
+		bool bIsEncrypted = false;
+
+		struct FEncryptionInfo
+		{
+			TArray<uint8> DefaultKID;
+			TArray<uint8> DefaultIV;
+			int32 DefaultIVSize = 0;
+			uint32 Scheme = 0;
+			struct FCDMInfo
+			{
+				TArray<uint8> SystemID;
+				TArray<TArray<uint8>> KIDs;
+				TArray<uint8> Data;
+			};
+			TArray<FCDMInfo> CDMInfos;
+		};
+
+	};
+
+	struct FDecoderInformation
+	{
+		bool bIsDecodable = true;
+		struct FProviderInfo
+		{
+			FString Name;
+			FString Version;
+			FString Implementation;
+			FString Vendor;
+		};
+		FProviderInfo ProviderInfo;
+	};
+
 } // namespace Electra
 
 #undef UE_API
diff --git a/Engine/Plugins/Media/ElectraUtil/Source/ElectraHTTPStream/Private/Generic/GenericElectraHTTPStream.cpp b/Engine/Plugins/Media/ElectraUtil/Source/ElectraHTTPStream/Private/Generic/GenericElectraHTTPStream.cpp
--- a/Engine/Plugins/Media/ElectraUtil/Source/ElectraHTTPStream/Private/Generic/GenericElectraHTTPStream.cpp
+++ b/Engine/Plugins/Media/ElectraUtil/Source/ElectraHTTPStream/Private/Generic/GenericElectraHTTPStream.cpp
@@ -79,7 +79,7 @@ public:
 
 	void AddRequest(IElectraHTTPStreamRequestPtr Request) override;
 
-	void TriggerWorkSignal()
+	void TriggerWorkSignal() override
 	{ HaveWorkSignal.Signal(); }
 
 private:
diff --git a/Engine/Plugins/Media/ElectraUtil/Source/ElectraHTTPStream/Public/ElectraHTTPStream.h b/Engine/Plugins/Media/ElectraUtil/Source/ElectraHTTPStream/Public/ElectraHTTPStream.h
--- a/Engine/Plugins/Media/ElectraUtil/Source/ElectraHTTPStream/Public/ElectraHTTPStream.h
+++ b/Engine/Plugins/Media/ElectraUtil/Source/ElectraHTTPStream/Public/ElectraHTTPStream.h
@@ -236,6 +236,12 @@ public:
 	virtual void AddThreadHandlerDelegate(FElectraHTTPStreamThreadHandlerDelegate InDelegate) = 0;
 	virtual void RemoveThreadHandlerDelegate() = 0;
 
+	/**
+	 * Triggers a call to the registered handler delegate without waiting for the worker thread's
+	 * timeout period.
+	 */
+	virtual void TriggerWorkSignal() = 0;
+
 	/**
 	 * Closes this handler and all currently active or pending requests.
 	 * This should be called before resetting the TSharedPtr<> returned by Create().
diff --git a/Engine/Plugins/Media/MP4Utilities/Source/MP4Boxes/Private/MP4Boxes.cpp b/Engine/Plugins/Media/MP4Utilities/Source/MP4Boxes/Private/MP4Boxes.cpp
--- a/Engine/Plugins/Media/MP4Utilities/Source/MP4Boxes/Private/MP4Boxes.cpp
+++ b/Engine/Plugins/Media/MP4Utilities/Source/MP4Boxes/Private/MP4Boxes.cpp
@@ -341,6 +341,7 @@ namespace MP4Boxes
 	struct FMP4BoxHDLR::FParsed
 	{
 		FString HandlerName;
+		uint32 Reserved[3] {};
 		uint32 HandlerType = 0;
 		uint32 Flags = 0;
 		uint8 Version = 0;
@@ -356,7 +357,9 @@ namespace MP4Boxes
 		ar.ReadVersionAndFlags(Parsed->Version, Parsed->Flags);
 		ar.SkipBytes(4);				// `pre_defined`
 		ar.Read(Parsed->HandlerType);
-		ar.SkipBytes(12);				// `reserved`
+		ar.Read(Parsed->Reserved[0]);
+		ar.Read(Parsed->Reserved[1]);
+		ar.Read(Parsed->Reserved[2]);
 		ar.ReadStringUTF8(Parsed->HandlerName, ar.GetNumBytesRemaining());
 	}
 
@@ -370,6 +373,12 @@ namespace MP4Boxes
 		ParseIfRequired();
 		return Parsed->HandlerName;
 	}
+	uint32 FMP4BoxHDLR::GetReservedValue(int32 InIndex0to2)
+	{
+		ParseIfRequired();
+		return InIndex0to2 >= 0 && InIndex0to2 <= 2 ? Parsed->Reserved[InIndex0to2] : 0;
+	}
+
 
 /****************************************************************************************************************************************************/
 
@@ -1489,6 +1498,16 @@ namespace MP4Boxes
 		ParseIfRequired();
 		return Parsed->Entries;
 	}
+	bool FMP4BoxSTTS::HasConstantSampleDuration()
+	{
+		ParseIfRequired();
+		return Parsed->Entries.Num() == 1;
+	}
+	int64 FMP4BoxSTTS::GetConstantSampleDuration()
+	{
+		ParseIfRequired();
+		return Parsed->Entries.Num() == 1 ? Parsed->Entries[0].sample_delta : 0;
+	}
 
 /****************************************************************************************************************************************************/
 
@@ -1662,13 +1681,13 @@ namespace MP4Boxes
 		{
 			return 0;
 		}
-		if (Parsed->SizeArray32)
+		if (Parsed->ConstantSampleSize)
 		{
-			return MP4Utilities::GetFromBigEndian(Parsed->SizeArray32[InIndex]);
+			return Parsed->ConstantSampleSize;
 		}
-		else if (Parsed->ConstantSampleSize)
+		else if (Parsed->SizeArray32)
 		{
-			return Parsed->ConstantSampleSize;
+			return MP4Utilities::GetFromBigEndian(Parsed->SizeArray32[InIndex]);
 		}
 		else if (Parsed->SizeArray16)
 		{
@@ -1685,6 +1704,18 @@ namespace MP4Boxes
 		}
 		return 0;
 	}
+	bool FMP4BoxSTSZ::HasConstantSampleSize()
+	{
+		ParseIfRequired();
+		return Parsed->ConstantSampleSize != 0;
+	}
+	uint32 FMP4BoxSTSZ::GetConstantSampleSize()
+	{
+		ParseIfRequired();
+		return Parsed->ConstantSampleSize;
+	}
+
+
 
 /****************************************************************************************************************************************************/
 
@@ -2304,6 +2335,7 @@ namespace MP4Boxes
 		TArray<uint32> SampleFlags;
 		TArray<int32> SampleCompositionTimeOffsets;
 		int64 TotalDuration = 0;
+		int64 TotalSampleSize = 0;
 		uint32 SampleCount = 0;
 		uint32 FirstSampleFlags = 0;
 		uint32 Flags = 0;
@@ -2357,7 +2389,10 @@ namespace MP4Boxes
 				}
 				if ((Parsed->Flags & 0x200) != 0)
 				{
-					ar.Read(Parsed->SampleSizes[i]);
+					uint32 sz;
+					ar.Read(sz);
+					Parsed->SampleSizes[i] = sz;
+					Parsed->TotalSampleSize += sz;
 				}
 				if ((Parsed->Flags & 0x400) != 0)
 				{
@@ -2437,6 +2472,12 @@ namespace MP4Boxes
 		ParseIfRequired();
 		return Parsed->SampleSizes;
 	}
+	int64 FMP4BoxTRUN::GetTotalSampleSize()
+	{
+		ParseIfRequired();
+		return Parsed->TotalSampleSize;
+	}
+
 	bool FMP4BoxTRUN::HasSampleFlags()
 	{
 		ParseIfRequired();
@@ -2808,6 +2849,55 @@ namespace MP4Boxes
 	}
 
 
+/****************************************************************************************************************************************************/
+	struct FMP4BoxMETA::FParsed
+	{
+		uint32 Flags = 0;
+		uint8 Version = 0;
+	};
+	void FMP4BoxMETA::ParseIfRequired()
+	{
+		if (Parsed.IsValid())
+		{
+			return;
+		}
+		Parsed = MakePimpl<FParsed>();
+		MP4Utilities::FMP4AtomReader ar(BoxInfo.Data);
+
+		// The `meta` box is a full box under ISO/IEC 14496-12 and a basic box in QuickTime.
+		// Some tools write a QT `meta` box even for ISO which is then missing the version number and flags!
+		// To work around this issue we try to locate the mandatory `hdlr` box that needs to be the first
+		// box and based on its position handle this as a basic or a full box.
+		bool bIsFullBox = false;
+		if (BoxInfo.Data.Num() >= 12)
+		{
+			bIsFullBox = BoxInfo.Data[8] == 'h' && BoxInfo.Data[9] == 'd' && BoxInfo.Data[10] == 'l' && BoxInfo.Data[11] == 'r';
+		}
+		if (bIsFullBox)
+		{
+			ar.ReadVersionAndFlags(Parsed->Version, Parsed->Flags);
+		}
+
+		while(ar.GetNumBytesRemaining() >= 8)
+		{
+			MP4Utilities::FMP4BoxInfo bi;
+			if (!ar.ParseIntoBoxInfo(bi, BoxInfo.Offset + BoxInfo.DataOffset + ar.GetCurrentOffset()))
+			{
+				UE_LOG(LogMP4Boxes, Error, TEXT("Failed to parse `meta` box!"));
+				return;
+			}
+			TSharedPtr<FMP4BoxBase, ESPMode::ThreadSafe> ChildBox = FMP4BoxFactory::Get().Create(AsWeak(), bi);
+			AddChildBox(MoveTemp(ChildBox));
+			ar.SkipBytes(bi.Data.Num());
+		}
+	}
+
+	TSharedPtr<FMP4BoxHDLR, ESPMode::ThreadSafe> FMP4BoxMETA::GetHandler()
+	{
+		ParseIfRequired();
+		return FindBoxRecursive<FMP4BoxHDLR>(MP4Utilities::MakeBoxAtom('h','d','l','r'), 0);
+	}
+
 /****************************************************************************************************************************************************/
 #if 0 // TEMPLATE
 	struct FMP4BoxCCCC::FParsed
diff --git a/Engine/Plugins/Media/MP4Utilities/Source/MP4Boxes/Private/MP4Track.cpp b/Engine/Plugins/Media/MP4Utilities/Source/MP4Boxes/Private/MP4Track.cpp
--- a/Engine/Plugins/Media/MP4Utilities/Source/MP4Boxes/Private/MP4Track.cpp
+++ b/Engine/Plugins/Media/MP4Utilities/Source/MP4Boxes/Private/MP4Track.cpp
@@ -542,7 +542,7 @@ namespace MP4Boxes
 		}
 		TUniquePtr<FIterator> It(new FIterator);
 		It->Track = AsWeak();
-		It->FragmentVars.TrunIterator.SetFragmentInfo(FragmentInfo);
+		It->FragmentVars.TrafTrunIterator.SetFragmentInfo(FragmentInfo);
 		// Copy convenience values into the iterator
 		It->Convs = Convs;
 		// Set up the box iterators.
@@ -711,7 +711,7 @@ namespace MP4Boxes
 		}
 		else
 		{
-			FragmentVars.TrunIterator.StartAt(InFrameNumber);
+			FragmentVars.TrafTrunIterator.StartAt(InFrameNumber);
 			FragmentVars.bInFragments = true;
 		}
 	}
@@ -767,9 +767,9 @@ namespace MP4Boxes
 		{
 			// We are in the fragment samples.
 			const int64 BaseMediaDecodeTime = 0;
-			int64 DTS = FragmentVars.TrunIterator.GetDTS() + BaseMediaDecodeTime;
-			int64 PTS = FragmentVars.TrunIterator.GetPTS() + BaseMediaDecodeTime;
-			int64 Duration = FragmentVars.TrunIterator.GetDuration();
+			int64 DTS = FragmentVars.TrafTrunIterator.GetDTS() + BaseMediaDecodeTime;
+			int64 PTS = FragmentVars.TrafTrunIterator.GetPTS() + BaseMediaDecodeTime;
+			int64 Duration = FragmentVars.TrafTrunIterator.GetDuration();
 			const uint32 Timescale = Convs.DurationFromMdhdBox.GetDenominator();
 			CurrentDTS.SetFromND(DTS, Timescale);
 			CurrentPTS.SetFromND(PTS, Timescale);
@@ -781,13 +781,13 @@ namespace MP4Boxes
 			// in timespan units. This is to avoid transformation issues from media local
 			// time into the timescale used in engine.
 			CurrentDurationTS = MP4Utilities::FFractionalTime(DTS + Duration, Timescale).GetAsTimespan() - CurrentDTS.GetAsTimespan();
-			CurrentSampleSize = FragmentVars.TrunIterator.GetSampleSize();
-			bCurrentIsSyncOrRAP = FragmentVars.TrunIterator.IsSyncSample();
-			CurrentSampleFileOffset = FragmentVars.TrunIterator.GetAbsoluteOffset();
-			SampleNumber = FragmentVars.TrunIterator.GetAbsoluteSampleNumber();
+			CurrentSampleSize = FragmentVars.TrafTrunIterator.GetSampleSize();
+			bCurrentIsSyncOrRAP = FragmentVars.TrafTrunIterator.IsSyncSample();
+			CurrentSampleFileOffset = FragmentVars.TrafTrunIterator.GetAbsoluteOffset();
+			SampleNumber = FragmentVars.TrafTrunIterator.GetAbsoluteSampleNumber();
 			// Subsample list
-			CurrentSubsampleList = FragmentVars.TrunIterator.GetSubsampleList();
-			CurrentSubsampleBoxFlags = FragmentVars.TrunIterator.GetSubsampleBoxFlags();
+			CurrentSubsampleList = FragmentVars.TrafTrunIterator.GetSubsampleList();
+			CurrentSubsampleBoxFlags = FragmentVars.TrafTrunIterator.GetSubsampleBoxFlags();
 		}
 	}
 
@@ -815,7 +815,7 @@ namespace MP4Boxes
 			// Already in fragments?
 			if (FragmentVars.bInFragments)
 			{
-				if (!FragmentVars.TrunIterator.Next())
+				if (!FragmentVars.TrafTrunIterator.Next())
 				{
 					check(!"why did we get here?");
 					return false;
@@ -851,7 +851,7 @@ namespace MP4Boxes
 		else
 		{
 			// Going back from fragments into the moov?
-			if (!FragmentVars.TrunIterator.Prev())
+			if (!FragmentVars.TrafTrunIterator.Prev())
 			{
 				StartAt(NextSampleNumber);
 			}
@@ -880,9 +880,63 @@ namespace MP4Boxes
 		return false;
 	}
 
+	bool FMP4Track::FIterator::GetCurrentChunkRemainingSampleInfo(FChunkInfo& OutInfo) const
+	{
+		if (!IsValid())
+		{
+			return false;
+		}
+		if (!FragmentVars.bInFragments)
+		{
+			uint32 NumRem = stscIt.GetNumSamplesInCurrentChunk() - stscIt.GetSampleIndexInCurrentChunk();
+			if (NumRem)
+			{
+				OutInfo.NumSamples = NumRem;
+				auto STSZ = stszIt.GetBox();
+				if (uint32 cs = STSZ->GetConstantSampleSize())
+				{
+					OutInfo.SizeInBytes = NumRem * cs;
+				}
+				else
+				{
+					OutInfo.SizeInBytes = 0;
+					for(uint32 i=stszIt.GetCurrentSampleNum(),iMax=i+NumRem; i<iMax; ++i)
+					{
+						OutInfo.SizeInBytes += STSZ->GetSizeOfSample(i);
+					}
+				}
+				if (int64 cd = sttsIt.GetConstantSampleDuration())
+				{
+					OutInfo.Duration = NumRem * cd;
+				}
+				else
+				{
+					OutInfo.Duration = 0;
+					FSTTSBoxIterator durIt(sttsIt);
+					for(uint32 i=0; i<NumRem; ++i)
+					{
+						OutInfo.Duration += durIt.GetCurrentDuration();
+						durIt.Next();
+					}
+				}
+			}
+			else
+			{
+				OutInfo = {};
+			}
+			// This is the last chunk if there are no fragments following that we could move into.
+			OutInfo.bIsLastChunk = stscIt.IsFinalChunk() && Convs.NumMoovSamples == Convs.NumTotalSamples;
+			return true;
+		}
+		else
+		{
+			return FragmentVars.TrafTrunIterator.GetCurrentTrafRemainingSampleInfo(OutInfo);
+		}
+	}
+
 
 
-	void FMP4Track::FIterator::FTrunIterator::SetFragmentInfo(const TSharedPtr<FFragmentInfo, ESPMode::ThreadSafe>& InFragmentInfo)
+	void FMP4Track::FIterator::FTrafTrunIterator::SetFragmentInfo(const TSharedPtr<FFragmentInfo, ESPMode::ThreadSafe>& InFragmentInfo)
 	{
 		FragmentInfo = InFragmentInfo;
 		if (FragmentInfo && FragmentInfo->TrexBox)
@@ -901,7 +955,7 @@ namespace MP4Boxes
 		}
 	}
 
-	bool FMP4Track::FIterator::FTrunIterator::StartAt(uint32 InSampleNum)
+	bool FMP4Track::FIterator::FTrafTrunIterator::StartAt(uint32 InSampleNum)
 	{
 		check(IsSetup());
 
@@ -962,7 +1016,7 @@ namespace MP4Boxes
 		return true;
 	}
 
-	void FMP4Track::FIterator::FTrunIterator::Update(const TSharedPtr<FFragmentInfo::FTrackTrafs, ESPMode::ThreadSafe>& InTraf)
+	void FMP4Track::FIterator::FTrafTrunIterator::Update(const TSharedPtr<FFragmentInfo::FTrackTrafs, ESPMode::ThreadSafe>& InTraf)
 	{
 		DefaultSampleDescriptionIndex = InTraf->TfhdBox->HasSampleDescriptionIndex() ? InTraf->TfhdBox->GetSampleDescriptionIndex() : TrexDefaultSampleDescriptionIndex;
 		DefaultSampleDuration = InTraf->TfhdBox->HasDefaultSampleDuration() ? InTraf->TfhdBox->GetDefaultSampleDuration() : TrexDefaultSampleDuration;
@@ -983,9 +1037,21 @@ namespace MP4Boxes
 		{
 			BaseOffset = InTraf->MoofBox->GetBoxFileOffset();
 		}
+
+		// Calculate the byte range of the first sample in the first `trun` up to and including the last sample in the last `trun`.
+		// This does not necessarily mean that all samples are consecutive in the file as each `trun` has its own offset so `trun`'s
+		// could as well be interleaved across sibling `traf`'s.
+		// We only want to know the file byte range all the samples are to be found in.
+		if (InTraf->TrunBoxes.Num())
+		{
+			TrunFirstSampleAbsoluteByteOffset = BaseOffset + (InTraf->TrunBoxes[0]->HasSampleOffset() ? InTraf->TrunBoxes[0]->GetSampleOffset() : 0);
+			TrunLastSampleAbsoluteByteOffset = BaseOffset + (InTraf->TrunBoxes.Last()->HasSampleOffset() ? InTraf->TrunBoxes.Last()->GetSampleOffset() : 0);
+			int64 LastSize = InTraf->TrunBoxes.Last()->HasSampleSizes() ? InTraf->TrunBoxes.Last()->GetTotalSampleSize() : InTraf->TrunBoxes.Last()->GetNumberOfSamples() * DefaultSampleSize;
+			TrunLastSampleAbsoluteByteOffset += LastSize;
+		}
 	}
 
-	void FMP4Track::FIterator::FTrunIterator::Update(const TSharedPtr<FMP4BoxTRUN, ESPMode::ThreadSafe>& InTrun)
+	void FMP4Track::FIterator::FTrafTrunIterator::Update(const TSharedPtr<FMP4BoxTRUN, ESPMode::ThreadSafe>& InTrun)
 	{
 		// Determine the offset to the sample by either adding the preceding sample sizes
 		// or directly if fixed sample size.
@@ -1040,7 +1106,7 @@ namespace MP4Boxes
 		SubsampleList = subsIt.GetCurrentSubsamples();
 	}
 
-	bool FMP4Track::FIterator::FTrunIterator::Prev()
+	bool FMP4Track::FIterator::FTrafTrunIterator::Prev()
 	{
 		check(IsSetup());
 		// If we are at the start of the fragments there is no previous sample to go back to.
@@ -1123,7 +1189,7 @@ namespace MP4Boxes
 		return false;
 	}
 
-	bool FMP4Track::FIterator::FTrunIterator::Next()
+	bool FMP4Track::FIterator::FTrafTrunIterator::Next()
 	{
 		check(IsSetup());
 		check(TrafNumber >= 0 && TrunNumber >= 0 && SampleNumber >= 0);
@@ -1183,5 +1249,47 @@ namespace MP4Boxes
 		return false;
 	}
 
+	bool FMP4Track::FIterator::FTrafTrunIterator::GetCurrentTrafRemainingSampleInfo(FChunkInfo& OutInfo) const
+	{
+		TSharedPtr<FFragmentInfo::FTrackTrafs, ESPMode::ThreadSafe> Traf = FragmentInfo->TrackFragments[TrafNumber];
+		TSharedPtr<FMP4BoxTRUN, ESPMode::ThreadSafe> Trun = Traf->TrunBoxes[TrunNumber];
+		const int64 LastTrafSampleNumber = Traf->AbsoluteFirstSampleNumber + Traf->NumSamples;
+		const int64 ThisSampleNumber = FirstAbsoluteSampleNumber + SampleNumber;
+		OutInfo.NumSamples = LastTrafSampleNumber - ThisSampleNumber;
+		//OutInfo.SizeInBytes = SampleSize;
+		OutInfo.SizeInBytes = TrunLastSampleAbsoluteByteOffset - GetAbsoluteOffset();
+		OutInfo.Duration = Duration;
+		OutInfo.bIsLastChunk = TrafNumber + 1 >= FragmentInfo->TrackFragments.Num();
+
+		uint32 NextTrunSampleNumber = (uint32)SampleNumber + 1;
+		int32 NextTrunNumber = TrunNumber + 1;
+		bool bTrunHasSampleSizes = Trun->HasSampleSizes();
+		bool bTrunHasSampleDurations = Trun->HasSampleDurations();
+		while(1)
+		{
+			// Still samples in this trun?
+			if (NextTrunSampleNumber < (int32)Trun->GetNumberOfSamples())
+			{
+				OutInfo.Duration += bTrunHasSampleDurations ? Trun->GetSampleDurations()[NextTrunSampleNumber] : DefaultSampleDuration;
+				//OutInfo.SizeInBytes += bTrunHasSampleSizes ? Trun->GetSampleSizes()[NextTrunSampleNumber] : DefaultSampleSize;
+				++NextTrunSampleNumber;
+			}
+			// Is there another trun following we can move into?
+			else if (NextTrunNumber < Traf->TrunBoxes.Num())
+			{
+				Trun = Traf->TrunBoxes[NextTrunNumber];
+				bTrunHasSampleSizes = Trun->HasSampleSizes();
+				bTrunHasSampleDurations = Trun->HasSampleDurations();
+				NextTrunSampleNumber = 0;
+				++NextTrunNumber;
+			}
+			else
+			{
+				break;
+			}
+		}
+		return true;
+	}
+
 
 } // namespace MP4Boxes
diff --git a/Engine/Plugins/Media/MP4Utilities/Source/MP4Boxes/Public/MP4BoxIterators.h b/Engine/Plugins/Media/MP4Utilities/Source/MP4Boxes/Public/MP4BoxIterators.h
--- a/Engine/Plugins/Media/MP4Utilities/Source/MP4Boxes/Public/MP4BoxIterators.h
+++ b/Engine/Plugins/Media/MP4Utilities/Source/MP4Boxes/Public/MP4BoxIterators.h
@@ -22,6 +22,7 @@ namespace MP4Boxes
 		{
 			SttsBox = InSttsBox;
 			NumTotalSamples = SttsBox->GetNumTotalSamples();
+			ConstantSampleDuration = SttsBox->GetConstantSampleDuration();
 			Time = 0;
 			CurrentSampleNum = 0;
 			LocalPosInCurrent = 0;
@@ -102,6 +103,8 @@ namespace MP4Boxes
 
 		int64 GetCurrentTime() const
 		{ return Time; }
+		int64 GetConstantSampleDuration() const
+		{ return ConstantSampleDuration; }
 		uint32 GetCurrentDuration() const
 		{ return CurrentDuration; }
 		uint32 GetCurrentSampleNum() const
@@ -116,6 +119,7 @@ namespace MP4Boxes
 	private:
 		TSharedPtr<FMP4BoxSTTS, ESPMode::ThreadSafe> SttsBox;
 		int64 Time = 0;
+		int64 ConstantSampleDuration = 0;
 		uint32 CurrentSampleNum = 0;
 		uint32 NumTotalSamples = 0;
 		uint32 LocalPosInCurrent = 0;
@@ -379,6 +383,16 @@ namespace MP4Boxes
 		bool IsLast() const
 		{ return CurrentSampleNum + 1 >= NumTotalSamples; }
 
+		bool IsFinalChunk() const
+		{
+			if (!StscBox.IsValid())
+			{
+				return true;
+			}
+			const uint32 NumSmpInLast = StscBox->GetEntries().Num() ? StscBox->GetEntries().Last().samples_per_chunk : 0;
+			const uint32 FirstSampleNumInLastChunk = NumTotalSamples - NumSmpInLast;
+			return CurrentSampleNum >= FirstSampleNumInLastChunk;
+		}
 	private:
 		TSharedPtr<FMP4BoxSTSC, ESPMode::ThreadSafe> StscBox;
 		uint32 Current_first_chunk = 0;
@@ -403,6 +417,11 @@ namespace MP4Boxes
 			CurrentSampleNum = 0;
 		}
 
+		TSharedPtr<FMP4BoxSTSZ, ESPMode::ThreadSafe> GetBox() const
+		{
+			return StszBox;
+		}
+
 		void SetToSampleNumber(uint32 InSampleNum)
 		{
 			CurrentSampleNum = InSampleNum < NumTotalSamples ? InSampleNum : NumTotalSamples-1;
diff --git a/Engine/Plugins/Media/MP4Utilities/Source/MP4Boxes/Public/MP4Boxes.h b/Engine/Plugins/Media/MP4Utilities/Source/MP4Boxes/Public/MP4Boxes.h
--- a/Engine/Plugins/Media/MP4Utilities/Source/MP4Boxes/Public/MP4Boxes.h
+++ b/Engine/Plugins/Media/MP4Utilities/Source/MP4Boxes/Public/MP4Boxes.h
@@ -298,6 +298,7 @@ namespace MP4Boxes
 
 		UE_API uint32 GetHandlerType();
 		UE_API FString GetHandlerName();
+		UE_API uint32 GetReservedValue(int32 InIndex0to2);
 	protected:
 		FMP4BoxHDLR(const TWeakPtr<FMP4BoxBase, ESPMode::ThreadSafe>& InParent, const MP4Utilities::FMP4BoxInfo& InBoxInfo)
 			: FMP4BoxBase(InParent, InBoxInfo)
@@ -794,6 +795,8 @@ namespace MP4Boxes
 		UE_API const TArray<FEntry>& GetEntries();
 		UE_API uint32 GetNumTotalSamples();
 		UE_API int64 GetTotalDuration();
+		UE_API bool HasConstantSampleDuration();
+		UE_API int64 GetConstantSampleDuration();
 	protected:
 		FMP4BoxSTTS(const TWeakPtr<FMP4BoxBase, ESPMode::ThreadSafe>& InParent, const MP4Utilities::FMP4BoxInfo& InBoxInfo)
 			: FMP4BoxBase(InParent, InBoxInfo)
@@ -938,6 +941,8 @@ namespace MP4Boxes
 		}
 		UE_API uint32 GetNumberOfSamples();
 		UE_API uint32 GetSizeOfSample(uint32 InIndex);
+		UE_API bool HasConstantSampleSize();
+		UE_API uint32 GetConstantSampleSize();
 	protected:
 		FMP4BoxSTSZ(const TWeakPtr<FMP4BoxBase, ESPMode::ThreadSafe>& InParent, const MP4Utilities::FMP4BoxInfo& InBoxInfo)
 			: FMP4BoxBase(InParent, InBoxInfo)
@@ -1312,6 +1317,7 @@ namespace MP4Boxes
 		UE_API const TArray<uint32>& GetSampleDurations();
 		UE_API bool HasSampleSizes();
 		UE_API const TArray<uint32>& GetSampleSizes();
+		UE_API int64 GetTotalSampleSize();
 		UE_API bool HasSampleFlags();
 		UE_API const TArray<uint32>& GetSampleFlags();
 		UE_API bool HasSampleCompositionTimeOffsets();
@@ -1567,10 +1573,14 @@ namespace MP4Boxes
 		virtual ~FMP4BoxMETA() = default;
 		bool IsLeafBox() const override { return false; }
 		bool IsListOfEntries() const override { return true; }
+		UE_API TSharedPtr<FMP4BoxHDLR, ESPMode::ThreadSafe> GetHandler();
 	protected:
 		FMP4BoxMETA(const TWeakPtr<FMP4BoxBase, ESPMode::ThreadSafe>& InParent, const MP4Utilities::FMP4BoxInfo& InBoxInfo)
 			: FMP4BoxBase(InParent, InBoxInfo)
 		{ }
+		void ParseIfRequired();
+		struct FParsed;
+		TPimplPtr<FParsed> Parsed;
 	};
 
 
diff --git a/Engine/Plugins/Media/MP4Utilities/Source/MP4Boxes/Public/MP4Track.h b/Engine/Plugins/Media/MP4Utilities/Source/MP4Boxes/Public/MP4Track.h
--- a/Engine/Plugins/Media/MP4Utilities/Source/MP4Boxes/Public/MP4Track.h
+++ b/Engine/Plugins/Media/MP4Utilities/Source/MP4Boxes/Public/MP4Track.h
@@ -180,14 +180,25 @@ namespace MP4Boxes
 			bool IsLastEffective() const
 			{ return SampleNumber >= Convs.LastSample.LastSampleNumber; }
 
+			struct FChunkInfo
+			{
+				uint32 NumSamples = 0;
+				int64 SizeInBytes = 0;
+				int64 Duration = 0;
+				bool bIsLastChunk = false;
+			};
+			// Returns information about the samples that are remaining in the current chunk.
+			// This is intended only for forward playback, not reverse.
+			UE_API bool GetCurrentChunkRemainingSampleInfo(FChunkInfo& OutInfo) const;
+
 			// Creates a copy of this iterator.
 			TSharedPtr<FIterator, ESPMode::ThreadSafe> Clone()
 			{ return MakeShareable(new FIterator(*this)); }
 		protected:
-			class FTrunIterator
+			class FTrafTrunIterator
 			{
 			public:
-				FTrunIterator() = default;
+				FTrafTrunIterator() = default;
 				void SetFragmentInfo(const TSharedPtr<FFragmentInfo, ESPMode::ThreadSafe>& InFragmentInfo);
 				bool IsSetup() const
 				{ return FragmentInfo.IsValid(); }
@@ -219,6 +230,8 @@ namespace MP4Boxes
 				{ return subsIt.GetSubsampleBoxFlags(); }
 				const FMP4BoxSUBS::FSubsampleList& GetSubsampleList() const
 				{ return SubsampleList; }
+
+				bool GetCurrentTrafRemainingSampleInfo(FChunkInfo& OutInfo) const;
 			private:
 				void Update(const TSharedPtr<FMP4BoxTRUN, ESPMode::ThreadSafe>& InTrun);
 				void Update(const TSharedPtr<FFragmentInfo::FTrackTrafs, ESPMode::ThreadSafe>& InTraf);
@@ -247,6 +260,9 @@ namespace MP4Boxes
 				int64 PTS = 0;
 				int64 Duration = 0;
 
+				int64 TrunFirstSampleAbsoluteByteOffset = -1;
+				int64 TrunLastSampleAbsoluteByteOffset = -1;
+
 				FSUBSBoxIterator subsIt;
 				FMP4BoxSUBS::FSubsampleList SubsampleList;
 			};
@@ -273,7 +289,7 @@ namespace MP4Boxes
 
 			struct FFragmentVars
 			{
-				FTrunIterator TrunIterator;
+				FTrafTrunIterator TrafTrunIterator;
 				bool bInFragments = false;
 			};
 			FFragmentVars FragmentVars;
@@ -315,6 +331,13 @@ namespace MP4Boxes
 		const MP4Utilities::FFractionalTime& GetFullMovieDuration() const
 		{ return Convs.FullMovieDuration; }
 
+		// Returns the duration of this track as mapped onto the media internal timeline (edit list).
+		const MP4Utilities::FFractionalTime& GetMappedTrackDuration() const
+		{ return Convs.MappedDurationFromElstBox; }
+		// Returns the duration of this track media.
+		const MP4Utilities::FFractionalTime& GetTrackDuration() const
+		{ return Convs.DurationFromMdhdBox; }
+
 	protected:
 		FMP4Track(const TSharedPtr<FMP4BoxTRAK, ESPMode::ThreadSafe>& InTrakBox, const TSharedPtr<FMP4Track::FFragmentInfo, ESPMode::ThreadSafe>& InFragmentInfo)
 			: TrakBox(InTrakBox)
