ElectraProtron: Added missing TracksChanged and Closed event when closing the player

#rnx

#Synced-CL 0

diff --git a/Engine/Plugins/Media/ElectraPlayer/Source/ElectraProtron/Private/Player/ElectraProtronPlayer.cpp b/Engine/Plugins/Media/ElectraPlayer/Source/ElectraProtron/Private/Player/ElectraProtronPlayer.cpp
--- a/Engine/Plugins/Media/ElectraPlayer/Source/ElectraProtron/Private/Player/ElectraProtronPlayer.cpp
+++ b/Engine/Plugins/Media/ElectraPlayer/Source/ElectraProtron/Private/Player/ElectraProtronPlayer.cpp
@@ -34,6 +34,8 @@ void FElectraProtronPlayer::Close()
 	{
 		ToClose->Close(FImpl::FCompletionDelegate());
 		ToClose.Reset();
+		EventSink.ReceiveMediaEvent(EMediaEvent::TracksChanged);
+		EventSink.ReceiveMediaEvent(EMediaEvent::MediaClosed);
 	}
 	CurrentTexturePool.Reset();
 	CurrentAudioSamplePool.Reset();
