Electra: Added a new mp4 handler path to the DASH and HLS segment reader; added common encryption handling to the mp4 boxes and the new mp4 handler

#rnx

#Synced-CL 51202517

diff --git a/Engine/Plugins/Media/ElectraPlayer/Source/ElectraPlayerRuntime/Private/Runtime/Player/AdaptiveStreamingPlayerEvents.h b/Engine/Plugins/Media/ElectraPlayer/Source/ElectraPlayerRuntime/Private/Runtime/Player/AdaptiveStreamingPlayerEvents.h
--- a/Engine/Plugins/Media/ElectraPlayer/Source/ElectraPlayerRuntime/Private/Runtime/Player/AdaptiveStreamingPlayerEvents.h
+++ b/Engine/Plugins/Media/ElectraPlayer/Source/ElectraPlayerRuntime/Private/Runtime/Player/AdaptiveStreamingPlayerEvents.h
@@ -3,6 +3,7 @@
 #pragma once
 
 #include "PlayerCore.h"
+#include "PlayerTime.h"
 
 namespace Electra
 {
@@ -45,7 +46,7 @@ public:
 
 /**
  * Client side event receiver.
- * 
+ *
  * As per ISO/IEC 23009-1:2019/DAM 1 event receivers subscribing to "on_receive" will receive the event the
  * instant it is received by the player. Handling of the event is the sole responsibility of the receiver.
  * Events will be received at any time and the same event can be received multiple times. Events will also
@@ -75,7 +76,7 @@ class IAdaptiveStreamingPlayerAEMSHandler
 public:
 	static IAdaptiveStreamingPlayerAEMSHandler* Create();
 	virtual ~IAdaptiveStreamingPlayerAEMSHandler() = default;
-	
+
 	// Call periodically to dispatch pending events.
 	virtual void ExecutePendingEvents(FTimeValue CurrentPlaybackTime) = 0;
 	// Unconditionally flushes all events from the queue.
diff --git a/Engine/Plugins/Media/ElectraPlayer/Source/ElectraPlayerRuntime/Private/Runtime/Player/AdaptiveStreamingPlayerManifest.cpp b/Engine/Plugins/Media/ElectraPlayer/Source/ElectraPlayerRuntime/Private/Runtime/Player/AdaptiveStreamingPlayerManifest.cpp
--- a/Engine/Plugins/Media/ElectraPlayer/Source/ElectraPlayerRuntime/Private/Runtime/Player/AdaptiveStreamingPlayerManifest.cpp
+++ b/Engine/Plugins/Media/ElectraPlayer/Source/ElectraPlayerRuntime/Private/Runtime/Player/AdaptiveStreamingPlayerManifest.cpp
@@ -20,7 +20,7 @@
 /*********************************************************************************************************************/
 
 #include "HAL/IConsoleManager.h"
-static TAutoConsoleVariable<bool> CVarElectraUseLegacyMP4PLayer(
+static TAutoConsoleVariable<bool> CVarElectraUseLegacyMP4Player(
 	TEXT("Electra.UseLegacyMP4Player"),
 	false,
 	TEXT("Use old sequentially reading mp4 player instead of new player.\n")
@@ -244,7 +244,7 @@ void FAdaptiveStreamingPlayer::InternalLoadManifest(const FString& InURL, const
 			}
 			else if (mimeType == Playlist::MIMETypeMP4 || mimeType == Playlist::MIMETypeMP4A || mimeType == Playlist::MIMETypeQuickTime)
 			{
-				if (CVarElectraUseLegacyMP4PLayer.GetValueOnAnyThread())
+				if (CVarElectraUseLegacyMP4Player.GetValueOnAnyThread())
 				{
 					ManifestReader = IPlaylistReaderMP4::Create(this);
 				}
diff --git a/Engine/Plugins/Media/ElectraPlayer/Source/ElectraPlayerRuntime/Private/Runtime/Player/DASH/ManifestDASH.cpp b/Engine/Plugins/Media/ElectraPlayer/Source/ElectraPlayerRuntime/Private/Runtime/Player/DASH/ManifestDASH.cpp
--- a/Engine/Plugins/Media/ElectraPlayer/Source/ElectraPlayerRuntime/Private/Runtime/Player/DASH/ManifestDASH.cpp
+++ b/Engine/Plugins/Media/ElectraPlayer/Source/ElectraPlayerRuntime/Private/Runtime/Player/DASH/ManifestDASH.cpp
@@ -7,6 +7,8 @@
 #include "PlaylistReaderDASH.h"
 #include "Stats/Stats.h"
 #include "Player/StreamSegmentReaderCommon.h"
+#include "MP4Boxes.h"
+#include "MP4Utilities.h"
 #include "Demuxer/ParserISO14496-12.h"
 #include "Demuxer/ParserISO14496-12_Utils.h"
 #include "Demuxer/ParserMKV.h"
@@ -58,39 +60,6 @@ namespace DashUtils
 	#undef GETPLAYEROPTION
 
 
-	/**
-	 * Helper class to parse a segment index (sidx box) from an ISO/IEC-14496:12 file.
-	 */
-	class FMP4SidxBoxReader : public FMP4StaticDataReader, public IParserISO14496_12::IBoxCallback
-	{
-	public:
-		FMP4SidxBoxReader() = default;
-		virtual ~FMP4SidxBoxReader() = default;
-	private:
-		//----------------------------------------------------------------------
-		// Methods from IParserISO14496_12::IBoxCallback
-		//
-		virtual EParseContinuation OnFoundBox(IParserISO14496_12::FBoxType Box, int64 BoxSizeInBytes, int64 FileDataOffset, int64 BoxDataOffset) override
-		{
-			if (bHaveSIDX)
-			{
-				return IParserISO14496_12::IBoxCallback::EParseContinuation::Stop;
-			}
-			if (Box == IParserISO14496_12::BoxType_sidx)
-			{
-				bHaveSIDX = true;
-			}
-			return IParserISO14496_12::IBoxCallback::EParseContinuation::Continue;
-		}
-		virtual EParseContinuation OnEndOfBox(IParserISO14496_12::FBoxType Box, int64 BoxSizeInBytes, int64 FileDataOffset, int64 BoxDataOffset) override
-		{
-			return IParserISO14496_12::IBoxCallback::EParseContinuation::Continue;
-		}
-		bool bHaveSIDX = false;
-	};
-
-
-
 	class FMatroskaParserDataReader : public IGenericDataReader, public TSharedFromThis<FMatroskaParserDataReader, ESPMode::ThreadSafe>
 	{
 	public:
@@ -2109,23 +2078,42 @@ void FDASHPlayPeriod::InitSegmentMP4DownloadComplete(TSharedPtrTS<FMPDLoadReques
 {
 	if (bSuccess)
 	{
-		FInitSegmentInfo::FAcceptBoxes AllBoxes;
-		FMP4StaticDataReader StaticDataReader;
-		StaticDataReader.SetParseData(LoadRequest->Request->GetResponseBuffer());
-		TSharedPtrTS<IParserISO14496_12> Init = IParserISO14496_12::CreateParser();
-		UEMediaError parseError = Init->ParseHeader(&StaticDataReader, &AllBoxes, PlayerSessionServices, nullptr);
-		if (parseError == UEMEDIA_ERROR_OK || parseError == UEMEDIA_ERROR_END_OF_STREAM)
+		if (Stream::CVarElectraStreamUseLegacyMP4.GetValueOnAnyThread())
 		{
-			// Parse the tracks of the init segment. We do this mainly to get to the CSD we might need should we have to insert filler data later.
-			parseError = Init->PrepareTracks(PlayerSessionServices, TSharedPtrTS<const IParserISO14496_12>());
-			if (parseError == UEMEDIA_ERROR_OK)
+			FInitSegmentInfo::FAcceptBoxes AllBoxes;
+			FMP4StaticDataReader StaticDataReader;
+			StaticDataReader.SetParseData(LoadRequest->Request->GetResponseBuffer());
+			TSharedPtrTS<IParserISO14496_12> Init = IParserISO14496_12::CreateParser();
+			UEMediaError parseError = Init->ParseHeader(&StaticDataReader, &AllBoxes, PlayerSessionServices, nullptr);
+			if (parseError == UEMEDIA_ERROR_OK || parseError == UEMEDIA_ERROR_END_OF_STREAM)
 			{
-				// Add this to the entity cache in case it needs to be retrieved again.
-				IPlayerEntityCache::FCacheItem CacheItem;
-				CacheItem.URL = LoadRequest->Url.URL;
-				CacheItem.Range = LoadRequest->Range;
-				CacheItem.Parsed14496_12Data = Init;
-				PlayerSessionServices->GetEntityCache()->CacheEntity(CacheItem);
+				// Parse the tracks of the init segment. We do this mainly to get to the CSD we might need should we have to insert filler data later.
+				parseError = Init->PrepareTracks(PlayerSessionServices, TSharedPtrTS<const IParserISO14496_12>());
+				if (parseError == UEMEDIA_ERROR_OK)
+				{
+					// Add this to the entity cache in case it needs to be retrieved again.
+					IPlayerEntityCache::FCacheItem CacheItem;
+					CacheItem.URL = LoadRequest->Url.URL;
+					CacheItem.Range = LoadRequest->Range;
+					CacheItem.Parsed14496_12Data = Init;
+					PlayerSessionServices->GetEntityCache()->CacheEntity(CacheItem);
+				}
+			}
+		}
+		else
+		{
+			TSharedPtrTS<IPlayerEntityCache> EntityCache = PlayerSessionServices ? PlayerSessionServices->GetEntityCache() : nullptr;
+			if (EntityCache)
+			{
+				TSharedPtr<FParsedMOOVCommon, ESPMode::ThreadSafe> ParsedMOOV = MakeShared<FParsedMOOVCommon, ESPMode::ThreadSafe>();
+				if (ParsedMOOV->Parse(MakeConstArrayView<uint8>(LoadRequest->Request->GetResponseBuffer()->GetLinearReadData(), LoadRequest->Request->GetResponseBuffer()->GetLinearReadSize())))
+				{
+					IPlayerEntityCache::FCacheItem CacheItem;
+					CacheItem.URL = LoadRequest->Url.URL;
+					CacheItem.Range = LoadRequest->Range;
+					CacheItem.ParsedMOOV = MoveTemp(ParsedMOOV);
+					EntityCache->CacheEntity(CacheItem);
+				}
 			}
 		}
 	}
@@ -2256,7 +2244,7 @@ FManifestDASHInternal::FRepresentation::ESearchResult FManifestDASHInternal::FRe
 	}
 
 	// Is a segment index (still) needed?
-	if (!SegmentIndexMP4.IsValid() && bNeedsSegmentIndex)
+	if (!SIDXBox.IsValid() && bNeedsSegmentIndex)
 	{
 		// Since this method may only be called with a still valid MPD representation we can pin again and don't need to check if it's still valid.
 		TSharedPtrTS<FDashMPD_RepresentationType> MPDRepresentation = Representation.Pin();
@@ -2346,7 +2334,7 @@ FManifestDASHInternal::FRepresentation::ESearchResult FManifestDASHInternal::FRe
 			if (InPlayerSessionServices->GetEntityCache()->GetCachedEntity(CachedItem, ChosenIndex.MediaURL.URL, RepresentationIndexSegmentURLRange))
 			{
 				// Already cached. Use it.
-				SegmentIndexMP4 = CachedItem.Parsed14496_12Data;
+				SIDXBox = StaticCastSharedPtr<MP4Boxes::FMP4BoxSIDX>(CachedItem.ISOBMFFBox);
 				SegmentMKV = CachedItem.ParsedMatroskaData;
 			}
 			else
@@ -2926,15 +2914,13 @@ FManifestDASHInternal::FRepresentation::ESearchResult FManifestDASHInternal::FRe
 	{
 		return SegIndexResult;
 	}
-	if (!SegmentIndexMP4.IsValid())
+	if (!SIDXBox.IsValid())
 	{
 		PostError(InPlayerSessionServices, FString::Printf(TEXT("A segment index is required for Representation \"%s\""), *MPDRepresentation->GetID()), ERRCODE_DASH_MPD_BAD_REPRESENTATION);
 		bIsUsable = false;
 		return FManifestDASHInternal::FRepresentation::ESearchResult::BadType;
 	}
-	const IParserISO14496_12::ISegmentIndex* Sidx = SegmentIndexMP4->GetSegmentIndexByIndex(0);
-	check(Sidx);	// The existence was already checked for in SegmentIndexDownloadComplete(), but just in case.
-	uint32 SidxTimescale = Sidx->GetTimescale();
+	uint32 SidxTimescale = SIDXBox->GetTimescale();
 	if (SidxTimescale == 0)
 	{
 		PostError(InPlayerSessionServices, FString::Printf(TEXT("Timescale of segment index for Representation \"%s\" is invalid!"), *MPDRepresentation->GetID()), ERRCODE_DASH_MPD_BAD_REPRESENTATION);
@@ -2976,9 +2962,9 @@ FManifestDASHInternal::FRepresentation::ESearchResult FManifestDASHInternal::FRe
 	//       The EPT also includes the very first composition time offset, so it may not be zero.
 	//       We have to remember that the segment index does not necessarily have access to an edit list as this is stored in the init segment and is not
 	//       available at this point in time, so any offsets that would come from an edit list need to have been applied to the EPT here already.
-	int64 EarliestPresentationTime = Sidx->GetEarliestPresentationTime();
+	int64 EarliestPresentationTime = SIDXBox->GetEarliestPresentationTime();
 	int64 CurrentT = EarliestPresentationTime;
-	int32 CurrentN, StartNumber=0, EndNumber=Sidx->GetNumEntries();
+	int32 CurrentN, StartNumber=0, EndNumber=SIDXBox->GetEntries().Num();
 	int32 CurrentD = 0;
 	int64 PreviousT = CurrentT;
 	int32 PreviousD = 0;
@@ -2988,7 +2974,7 @@ FManifestDASHInternal::FRepresentation::ESearchResult FManifestDASHInternal::FRe
 	bool bFound = false;
 	for(CurrentN=0; CurrentN<EndNumber; ++CurrentN)
 	{
-		const IParserISO14496_12::ISegmentIndex::FEntry& SegmentInfo = Sidx->GetEntry(CurrentN);
+		const MP4Boxes::FMP4BoxSIDX::FEntry& SegmentInfo = SIDXBox->GetEntries()[CurrentN];
 
 		// We do not support hierarchical segment indices!
 		if (SegmentInfo.IsReferenceType)
@@ -3113,7 +3099,7 @@ FManifestDASHInternal::FRepresentation::ESearchResult FManifestDASHInternal::FRe
 			--CurrentN;
 			CurrentD = PreviousD;
 			CurrentT = PreviousT;
-			CurrentOffset -= Sidx->GetEntry(CurrentN).Size;
+			CurrentOffset -= SIDXBox->GetEntries()[CurrentN].Size;
 			bFound = true;
 		}
 	}
@@ -3125,8 +3111,8 @@ FManifestDASHInternal::FRepresentation::ESearchResult FManifestDASHInternal::FRe
 		OutSegmentInfo.PTO = PTO;
 		OutSegmentInfo.Duration = CurrentD;
 		OutSegmentInfo.Number = CurrentN;
-		OutSegmentInfo.NumberOfBytes = Sidx->GetEntry(CurrentN).Size;
-		OutSegmentInfo.FirstByteOffset = Sidx->GetFirstOffset() + SegmentIndexRangeStart + SegmentIndexRangeSize + CurrentOffset;
+		OutSegmentInfo.NumberOfBytes = SIDXBox->GetEntries()[CurrentN].Size;
+		OutSegmentInfo.FirstByteOffset = SIDXBox->GetFirstOffset() + SegmentIndexRangeStart + SegmentIndexRangeSize + CurrentOffset;
 		OutSegmentInfo.MediaLocalFirstAUTime = OutSegmentInfo.MediaLocalFirstPTS = MediaLocalSearchTime;
 		OutSegmentInfo.MediaLocalLastAUTime = MediaLocalEndTime;
 		OutSegmentInfo.Timescale = SidxTimescale;
@@ -3140,8 +3126,8 @@ FManifestDASHInternal::FRepresentation::ESearchResult FManifestDASHInternal::FRe
 			int64 nlt = CurrentT + CurrentD;
 			for(int32 jj=CurrentN+1; jj<EndNumber && nlt<MediaLocalEndTime; ++jj)
 			{
-				int64 nss = Sidx->GetEntry(jj).Size;
-				int64 nsd = Sidx->GetEntry(jj).SubSegmentDuration;
+				int64 nss = SIDXBox->GetEntries()[jj].Size;
+				int64 nsd = SIDXBox->GetEntries()[jj].SubSegmentDuration;
 				int64 NextSegStart = cbo;
 				int64 NextSegEnd = cbo + nss - 1;
 				OutSegmentInfo.NextMediaURLS.Emplace_GetRef().Range = FString::Printf(TEXT("%lld-%lld"), (long long int)NextSegStart, (long long int)NextSegEnd);
@@ -3908,7 +3894,7 @@ void FManifestDASHInternal::FRepresentation::GetAverageSegmentDuration(FTimeValu
 				SegmentTemplate.Remove(nullptr);
 				if (SegmentBase.Num())
 				{
-					TSharedPtrTS<const IParserISO14496_12> SI = SegmentIndexMP4;
+					TSharedPtr<MP4Boxes::FMP4BoxSIDX, ESPMode::ThreadSafe> SI = SIDXBox;
 					// If the segment index on this representation is not there we look for any segment index of another representation.
 					// Since they need to be segmented the same we can at least use the segment durations from there.
 					bool bExact = true;
@@ -3919,29 +3905,25 @@ void FManifestDASHInternal::FRepresentation::GetAverageSegmentDuration(FTimeValu
 						for(int32 nR=0; nR<ParentAdaptation->GetNumberOfRepresentations(); ++nR)
 						{
 							FRepresentation* Rep = static_cast<FRepresentation*>(ParentAdaptation->GetRepresentationByIndex(nR).Get());
-							if (Rep->SegmentIndexMP4.IsValid())
+							if (Rep->SIDXBox.IsValid())
 							{
-								SI = Rep->SegmentIndexMP4;
+								SI = Rep->SIDXBox;
 								break;
 							}
 						}
 					}
 					if (SI.IsValid())
 					{
-						const IParserISO14496_12::ISegmentIndex* Sidx = SI->GetSegmentIndexByIndex(0);
-						if (Sidx)
+						uint32 SidxTimescale = SI->GetTimescale();
+						if (SidxTimescale)
 						{
-							uint32 SidxTimescale = Sidx->GetTimescale();
-							if (SidxTimescale)
+							int64 DurationSum = 0;
+							for(int32 nI=0, nIMax=SI->GetEntries().Num(); nI<nIMax; ++nI)
 							{
-								int64 DurationSum = 0;
-								for(int32 nI=0, nIMax=Sidx->GetNumEntries(); nI<nIMax; ++nI)
-								{
-									DurationSum += Sidx->GetEntry(nI).SubSegmentDuration;
-								}
-								OutAverageSegmentDuration.SetFromND(DurationSum / Sidx->GetNumEntries(), SidxTimescale);
-								return;
+								DurationSum += SI->GetEntries()[nI].SubSegmentDuration;
 							}
+							OutAverageSegmentDuration.SetFromND(DurationSum / SI->GetEntries().Num(), SidxTimescale);
+							return;
 						}
 					}
 				}
@@ -4006,23 +3988,22 @@ void FManifestDASHInternal::FRepresentation::SegmentIndexDownloadComplete(TShare
 		{
 			if (LoadRequest->OwningManifest.Pin().IsValid() && LoadRequest->XLinkElement == Representation)
 			{
-				DashUtils::FMP4SidxBoxReader BoxReader;
-				BoxReader.SetParseData(LoadRequest->Request->GetResponseBuffer());
-				TSharedPtrTS<IParserISO14496_12> Index = IParserISO14496_12::CreateParser();
-				UEMediaError parseError = Index->ParseHeader(&BoxReader, &BoxReader, LoadRequest->PlayerSessionServices, nullptr);
-				if (parseError == UEMEDIA_ERROR_OK || parseError == UEMEDIA_ERROR_END_OF_STREAM)
+				TSharedPtr<MP4Utilities::FMP4BoxData, ESPMode::ThreadSafe> bd = MP4Utilities::CreateBoxDataFromBuffer(MakeConstArrayView<uint8>(LoadRequest->Request->GetResponseBuffer()->GetLinearReadData(), LoadRequest->Request->GetResponseBuffer()->GetLinearReadSize()), 0);
+				MP4Boxes::FMP4BoxTreeParser tp;
+				if (bd && tp.ParseBoxTree(bd))
 				{
-					if (Index->PrepareTracks(LoadRequest->PlayerSessionServices, nullptr) == UEMEDIA_ERROR_OK && Index->GetNumberOfSegmentIndices() > 0)
+					TSharedPtr<MP4Boxes::FMP4BoxSIDX, ESPMode::ThreadSafe> sidx = StaticCastSharedPtr<MP4Boxes::FMP4BoxSIDX>(tp.GetBoxTree());
+					if (sidx->GetEntries().Num() > 0)
 					{
-						SegmentIndexMP4 = MoveTemp(Index);
+						SIDXBox = MoveTemp(sidx);
 						bOk = true;
 						// Add this to the entity cache in case it needs to be retrieved again.
 						IPlayerEntityCache::FCacheItem ci;
 						ci.URL = LoadRequest->Url.URL;
 						ci.Range = LoadRequest->Range;
-						ci.Parsed14496_12Data = SegmentIndexMP4;
+						ci.ISOBMFFBox = SIDXBox;
 						LoadRequest->PlayerSessionServices->GetEntityCache()->CacheEntity(ci);
-					}
+  					}
 					else
 					{
 						LogMessage(LoadRequest->PlayerSessionServices, IInfoLog::ELevel::Warning, FString::Printf(TEXT("Representation segment index is invalid. Marking representation as unusable.")));
diff --git a/Engine/Plugins/Media/ElectraPlayer/Source/ElectraPlayerRuntime/Private/Runtime/Player/DASH/PlayerEventDASH.h b/Engine/Plugins/Media/ElectraPlayer/Source/ElectraPlayerRuntime/Private/Runtime/Player/DASH/PlayerEventDASH.h
--- a/Engine/Plugins/Media/ElectraPlayer/Source/ElectraPlayerRuntime/Private/Runtime/Player/DASH/PlayerEventDASH.h
+++ b/Engine/Plugins/Media/ElectraPlayer/Source/ElectraPlayerRuntime/Private/Runtime/Player/DASH/PlayerEventDASH.h
@@ -4,6 +4,7 @@
 
 #include "PlayerCore.h"
 #include "Player/AdaptiveStreamingPlayerEvents.h"
+#include "Utilities/StringHelpers.h"
 #include "Misc/Base64.h"
 
 namespace Electra
@@ -44,7 +45,7 @@ public:
 			StringHelpers::StringToArray(MessageData, InData);
 		}
 	}
-	void SetMessageData(const TArray<uint8>& InData)
+	void SetMessageData(TConstArrayView<uint8> InData)
 	{
 		MessageData = InData;
 	}
diff --git a/Engine/Plugins/Media/ElectraPlayer/Source/ElectraPlayerRuntime/Private/Runtime/Player/DASH/PlaylistReaderDASH_Internal.h b/Engine/Plugins/Media/ElectraPlayer/Source/ElectraPlayerRuntime/Private/Runtime/Player/DASH/PlaylistReaderDASH_Internal.h
--- a/Engine/Plugins/Media/ElectraPlayer/Source/ElectraPlayerRuntime/Private/Runtime/Player/DASH/PlaylistReaderDASH_Internal.h
+++ b/Engine/Plugins/Media/ElectraPlayer/Source/ElectraPlayerRuntime/Private/Runtime/Player/DASH/PlaylistReaderDASH_Internal.h
@@ -2,7 +2,7 @@
 
 #pragma once
 
-#include "CoreMinimal.h"
+#include "CoreTypes.h"
 #include "PlayerCore.h"
 
 #include "MPDElementsDASH.h"
@@ -12,6 +12,7 @@
 #include "Player/StreamSegmentReaderCommon.h"
 #include "Player/ContentSteeringHandler.h"
 #include "Utilities/URLParser.h"
+#include "MP4Boxes.h"
 
 
 namespace Electra
@@ -274,7 +275,7 @@ public:
 		bool bWarnedAboutInconsistentAvailabilityTimeComplete = false;
 		//
 		bool bNeedsSegmentIndex = true;
-		TSharedPtrTS<const IParserISO14496_12> SegmentIndexMP4;
+		TSharedPtr<MP4Boxes::FMP4BoxSIDX, ESPMode::ThreadSafe> SIDXBox;
 		TSharedPtrTS<const IParserMKV> SegmentMKV;
 		int64 SegmentIndexRangeStart = 0;
 		int64 SegmentIndexRangeSize = 0;
diff --git a/Engine/Plugins/Media/ElectraPlayer/Source/ElectraPlayerRuntime/Private/Runtime/Player/PlayerEntityCache.h b/Engine/Plugins/Media/ElectraPlayer/Source/ElectraPlayerRuntime/Private/Runtime/Player/PlayerEntityCache.h
--- a/Engine/Plugins/Media/ElectraPlayer/Source/ElectraPlayerRuntime/Private/Runtime/Player/PlayerEntityCache.h
+++ b/Engine/Plugins/Media/ElectraPlayer/Source/ElectraPlayerRuntime/Private/Runtime/Player/PlayerEntityCache.h
@@ -6,11 +6,13 @@
 #include "PlayerTime.h"
 #include "ParameterDictionary.h"
 #include "HTTP/HTTPManager.h"
+#include "MP4BoxBase.h"
 
 namespace Electra
 {
 class IPlayerSessionServices;
 class IParserISO14496_12;
+class FParsedMOOVCommon;
 class IParserMKV;
 
 
@@ -32,6 +34,8 @@ public:
 		FString URL;
 		FString Range;
 		FTimeValue ExpiresAtUTC;
+		TSharedPtr<MP4Boxes::FMP4BoxBase, ESPMode::ThreadSafe> ISOBMFFBox;
+		TSharedPtr<FParsedMOOVCommon, ESPMode::ThreadSafe> ParsedMOOV;
 		TSharedPtrTS<const IParserISO14496_12>	Parsed14496_12Data;
 		TSharedPtrTS<const IParserMKV>	ParsedMatroskaData;
 		TSharedPtrTS<const TArray<uint8>> RawPayloadData;
diff --git a/Engine/Plugins/Media/ElectraPlayer/Source/ElectraPlayerRuntime/Private/Runtime/Player/StreamSegmentReaderCommon.cpp b/Engine/Plugins/Media/ElectraPlayer/Source/ElectraPlayerRuntime/Private/Runtime/Player/StreamSegmentReaderCommon.cpp
--- a/Engine/Plugins/Media/ElectraPlayer/Source/ElectraPlayerRuntime/Private/Runtime/Player/StreamSegmentReaderCommon.cpp
+++ b/Engine/Plugins/Media/ElectraPlayer/Source/ElectraPlayerRuntime/Private/Runtime/Player/StreamSegmentReaderCommon.cpp
@@ -3,7 +3,9 @@
 #include "StreamSegmentReaderCommon.h"
 #include "PlayerCore.h"
 #include "HTTP/HTTPManager.h"
+#include "HTTP/HTTPResponseCache.h"
 #include "MP4Utilities.h"
+#include "MP4Track.h"
 #include "Demuxer/ParserISO14496-12.h"
 #include "Demuxer/ParserISO14496-12_Utils.h"
 #include "Demuxer/ParserISO13818-1.h"
@@ -26,6 +28,7 @@
 #include "Utilities/UtilsMP4.h"
 #include "Utilities/UtilsMPEG.h"
 #include "Utilities/ElectraBitstream.h"
+#include "Utilities/MP4Helpers.h"
 #include "Utils/Google/ElectraUtilsVPxVideo.h"
 #include "Utils/MPEG/ElectraUtilsMPEGVideo_H264.h"
 #include "Utils/MPEG/ElectraUtilsMPEGVideo_H265.h"
@@ -37,6 +40,20 @@
 
 DECLARE_CYCLE_STAT(TEXT("FStreamSegmentReaderCommon::HandleRequest"), STAT_ElectraPlayer_Common_StreamReader, STATGROUP_ElectraPlayer);
 
+
+/*********************************************************************************************************************/
+
+namespace Electra { namespace Stream {
+TAutoConsoleVariable<bool> CVarElectraStreamUseLegacyMP4(
+	TEXT("Electra.UseLegacyMP4Stream"),
+	false,
+	TEXT("Use old mp4 parser for DASH and HLS streaming.\n")
+	TEXT(" true: use old parser, false use new parser"),
+	ECVF_Default);
+}}
+/*********************************************************************************************************************/
+
+
 namespace Electra
 {
 
@@ -67,6 +84,7 @@
 		bool HasReachedEOF() const override;
 
 		void SetSourceBuffer(TSharedPtrTS<FWaitableBuffer> InDataBuffer);
+		void SetSourceBuffer(TArray<uint8>&& InData);
 		bool HaveSourceBuffer() const;
 		void Abort();
 		bool HasErrored() const;
@@ -74,8 +92,11 @@
 		bool GetEOD() const;
 		const uint8* GetBufferBaseAddress() const;
 
+		TConstArrayView<uint8> GetConstArrayView() const
+		{ return DataBuffer ? MakeConstArrayView<uint8>(DataBuffer->GetLinearReadData(), DataBuffer->GetLinearReadSize()) : Data; }
 	private:
 		TSharedPtrTS<FWaitableBuffer> DataBuffer;
+		TArray<uint8> Data;
 		std::atomic<int64> CurrentPos {0};
 		volatile bool bAborted = false;
 	};
@@ -87,6 +108,14 @@
 		bAborted = false;
 	}
 
+	void FGenericDataReader::SetSourceBuffer(TArray<uint8>&& InData)
+	{
+		DataBuffer.Reset();
+		Data = MoveTemp(InData);
+		CurrentPos = 0;
+		bAborted = false;
+	}
+
 	bool FGenericDataReader::HaveSourceBuffer() const
 	{
 		return DataBuffer.IsValid();
@@ -175,6 +204,8 @@
 
 
 
+
+
 FStreamSegmentRequestCommon::FStreamSegmentRequestCommon()
 {
 }
@@ -578,13 +609,14 @@
 	StreamSelector.Reset();
 }
 
-void FStreamSegmentReaderCommon::FStreamHandler::SetError(const FString& InMessage, uint16 InCode)
+FStreamSegmentReaderCommon::FStreamHandler::EHandleResult FStreamSegmentReaderCommon::FStreamHandler::SetError(const FString& InMessage, uint16 InCode)
 {
 	SegmentError.SetError(UEMEDIA_ERROR_FORMAT_ERROR);
 	SegmentError.SetFacility(Facility::EFacility::CommonSegmentReader);
 	SegmentError.SetCode(InCode);
 	SegmentError.SetMessage(InMessage);
 	bHasErrored = true;
+	return EHandleResult::Failed;
 }
 
 bool FStreamSegmentReaderCommon::FStreamHandler::HasErrored() const
@@ -705,7 +737,7 @@
 		{
 			if (Request->ExpectedContainerFormat == FStreamSegmentRequestCommon::EContainerFormat::ISO_14496_12)
 			{
-				SegmentHandleResult = HandleMP4Media(Request);
+				SegmentHandleResult = Stream::CVarElectraStreamUseLegacyMP4.GetValueOnAnyThread() ? HandleMP4MediaLegacy(Request) : HandleMP4Media(Request);
 			}
 			else if (Request->ExpectedContainerFormat == FStreamSegmentRequestCommon::EContainerFormat::Matroska_WebM)
 			{
@@ -819,7 +851,7 @@
 
 			if (ContainerFormat == FStreamSegmentRequestCommon::EContainerFormat::ISO_14496_12)
 			{
-				SegmentHandleResult = HandleMP4Media(Request);
+				SegmentHandleResult = Stream::CVarElectraStreamUseLegacyMP4.GetValueOnAnyThread() ? HandleMP4MediaLegacy(Request) : HandleMP4Media(Request);
 			}
 			else if (ContainerFormat == FStreamSegmentRequestCommon::EContainerFormat::Matroska_WebM)
 			{
@@ -944,7 +976,12 @@
 		IPlayerEntityCache::FCacheItem CachedItem;
 		if (EntityCache->GetCachedEntity(CachedItem, InitURLToUse.Url.URL, InitURLToUse.Range))
 		{
-			if (CachedItem.Parsed14496_12Data.IsValid())
+			if (CachedItem.ParsedMOOV.IsValid())
 			{
+				InitSegmentData.Emplace<TSharedPtrTS<FParsedMOOVCommon>>(CachedItem.ParsedMOOV);
+				DownloadStats.bWasSuccessful = true;
+			}
+			else if (CachedItem.Parsed14496_12Data.IsValid())
+			{
 				InitSegmentData.Emplace<TSharedPtrTS<const IParserISO14496_12>>(CachedItem.Parsed14496_12Data);
 				DownloadStats.bWasSuccessful = true;
@@ -1136,55 +1173,113 @@
 	}
 	else
 	{
-		// For self-initializing media we need to locate the 'moov' box.
-		UtilsMP4::FMP4RootBoxLocator BoxLocator;
-		TArray<UtilsMP4::FMP4RootBoxLocator::FBoxInfo> Boxes;
-		const TArray<uint32> FirstBox { MP4Utilities::MakeBoxAtom('f','t','y','p') };
-		const TArray<uint32> MoovBox { MP4Utilities::MakeBoxAtom('m','o','o','v') };
+		// For self-initializing media we need to locate the `moov` box.
 		FString InitSegmentUrl(InitURLToUse.Url.URL);
-
 		// Set up information for CMCD.
 		FCMCDHandler::FRequestObjectInfo ri;
 		ri.MediaStreamType = InRequest->StreamType;
 		ri.EncodedBitrate = InRequest->GetBitrate();
 		PlayerSessionService->GetCMCDHandler()->SetupRequestObject(FCMCDHandler::ERequestType::InitSegment, FCMCDHandler::EObjectType::InitSegment, InitSegmentUrl, ReqHeaders, InitURLToUse.Url.CDN, ri);
-		bool bSuccess = BoxLocator.LocateRootBoxes(Boxes, PlayerSessionService->GetHTTPManager(), InitSegmentUrl, ReqHeaders, FirstBox, MoovBox, MoovBox, UtilsMP4::FMP4RootBoxLocator::FCancellationCheckDelegate::CreateLambda([&]()
+
+		if (!Stream::CVarElectraStreamUseLegacyMP4.GetValueOnAnyThread())
 		{
-			return HasReadBeenAborted();
-		}));
-		if (HasReadBeenAborted())
-		{
-			return true;
+  			TSharedPtr<FMP4DataLoader, ESPMode::ThreadSafe> DataLoader = FMP4DataLoader::Create();
+			FMP4DataLoader::FHttpParams HttpParams;
+			HttpParams.HTTPManager = PlayerSessionService->GetHTTPManager();
+			HttpParams.HTTPResponseCache = IHTTPResponseCache::Create(1024*1024, 32, nullptr);
+			HttpParams.RequestHeaders = ReqHeaders;
+			HttpParams.ConnectTimeoutMillis = 5000;
+			HttpParams.NoDataTimeoutMillis = 2000;
+			DataLoader->Open(InitSegmentUrl, HttpParams);
+
+			MP4Utilities::FMP4BoxLocator BoxLocator;
+			TArray<TSharedPtr<MP4Utilities::FMP4BoxData, ESPMode::ThreadSafe>> RootBoxes;
+			const TArray<uint32> StopAfterBoxes { MP4Utilities::MakeBoxAtom('m','o','o','v') };
+			bool bGotBoxes = BoxLocator.LocateAndReadRootBoxes(RootBoxes, DataLoader, TArray<uint32>(), StopAfterBoxes, false, TArray<uint32>(), FMP4DataLoader::FCancellationCheckDelegate::CreateLambda([&](){return HasReadBeenAborted();}));
+			if (HasReadBeenAborted())
+			{
+				return true;
+			}
+			CurrentConnectionInfo = DataLoader->GetConnectionInfo();
+			// If CMCD added query parameters to the URL then we need to remove them from the effective URL.
+			if (!InitSegmentUrl.Equals(InitURLToUse.Url.URL))
+			{
+				// Keep the initial URL though as it was since it was (a) not modified and (b) it is what the init segment gets cached under.
+				// We do not want to use the redirected effective URL for the cache.
+				CurrentConnectionInfo.EffectiveURL = PlayerSessionService->GetCMCDHandler()->RemoveParamsFromURL(CurrentConnectionInfo.EffectiveURL);
+			}
+			SetupSegmentDownloadStatsFromConnectionInfo(CurrentConnectionInfo);
+			if (!bGotBoxes)
+			{
+				FString ErrMsg = BoxLocator.GetLastError();
+				if (ErrMsg.IsEmpty())
+				{
+					ErrMsg = CurrentConnectionInfo.StatusInfo.ErrorDetail.GetMessage();
+				}
+				SetError(FString::Printf(TEXT("Self-initializing media `moov` scan error: %s"),	*ErrMsg), INTERNAL_SEG_ERROR_INIT_SEGMENT_DOWNLOAD_ERROR);
+				return false;
+			}
+			bool bHaveMoov = false;
+			for(int32 i=0; i<RootBoxes.Num(); ++i)
+			{
+				if (RootBoxes[i]->Type == MP4Utilities::MakeBoxAtom('m','o','o','v'))
+				{
+					StaticSegmentDataReader.SetSourceBuffer(RootBoxes[i]->GetBoxDataRAW());
+					ContainerFormat = FStreamSegmentRequestCommon::EContainerFormat::ISO_14496_12;
+					bHaveMoov = true;
+					break;
+				}
+			}
+			if (!bHaveMoov)
+			{
+				SetError(FString::Printf(TEXT("Self-initializing media does not have a 'moov' box")), INTERNAL_SEG_ERROR_INIT_SEGMENT_NOTFOUND_ERROR);
+				return false;
+			}
 		}
-		CurrentConnectionInfo = BoxLocator.GetConnectionInfo();
-		// If CMCD added query parameters to the URL then we need to remove them from the effective URL.
-		if (!InitSegmentUrl.Equals(InitURLToUse.Url.URL))
-		{
-			// Keep the initial URL though as it was since it was (a) not modified and (b) it is what the init segment gets cached under.
-			// We do not want to use the redirected effective URL for the cache.
-			CurrentConnectionInfo.EffectiveURL = PlayerSessionService->GetCMCDHandler()->RemoveParamsFromURL(CurrentConnectionInfo.EffectiveURL);
-		}
-		SetupSegmentDownloadStatsFromConnectionInfo(CurrentConnectionInfo);
-		if (!bSuccess)
-		{
-			FString ErrMsg = BoxLocator.GetErrorMessage();
-			if (ErrMsg.IsEmpty())
-			{
-				ErrMsg = CurrentConnectionInfo.StatusInfo.ErrorDetail.GetMessage();
-			}
-			SetError(FString::Printf(TEXT("Self-initializing media 'moov' scan error: %s"),	*ErrMsg), INTERNAL_SEG_ERROR_INIT_SEGMENT_DOWNLOAD_ERROR);
-			return false;
-		}
-		const UtilsMP4::FMP4RootBoxLocator::FBoxInfo* Moov = Boxes.FindByPredicate([InType=MP4Utilities::MakeBoxAtom('m','o','o','v')](const UtilsMP4::FMP4RootBoxLocator::FBoxInfo& InBox){return InBox.Type == InType;});
-		if (Moov)
-		{
-			StaticSegmentDataReader.SetSourceBuffer(Moov->DataBuffer);
-			ContainerFormat = FStreamSegmentRequestCommon::EContainerFormat::ISO_14496_12;
-		}
 		else
 		{
-			SetError(FString::Printf(TEXT("Self-initializing media does not have a 'moov' box")), INTERNAL_SEG_ERROR_INIT_SEGMENT_NOTFOUND_ERROR);
-			return false;
+			UtilsMP4::FMP4RootBoxLocator BoxLocator;
+			TArray<UtilsMP4::FMP4RootBoxLocator::FBoxInfo> Boxes;
+			const TArray<uint32> FirstBox { MP4Utilities::MakeBoxAtom('f','t','y','p') };
+			const TArray<uint32> MoovBox { MP4Utilities::MakeBoxAtom('m','o','o','v') };
+			bool bSuccess = BoxLocator.LocateRootBoxes(Boxes, PlayerSessionService->GetHTTPManager(), InitSegmentUrl, ReqHeaders, FirstBox, MoovBox, MoovBox, UtilsMP4::FMP4RootBoxLocator::FCancellationCheckDelegate::CreateLambda([&]()
+			{
+				return HasReadBeenAborted();
+			}));
+			if (HasReadBeenAborted())
+			{
+				return true;
+			}
+			CurrentConnectionInfo = BoxLocator.GetConnectionInfo();
+			// If CMCD added query parameters to the URL then we need to remove them from the effective URL.
+			if (!InitSegmentUrl.Equals(InitURLToUse.Url.URL))
+			{
+				// Keep the initial URL though as it was since it was (a) not modified and (b) it is what the init segment gets cached under.
+				// We do not want to use the redirected effective URL for the cache.
+				CurrentConnectionInfo.EffectiveURL = PlayerSessionService->GetCMCDHandler()->RemoveParamsFromURL(CurrentConnectionInfo.EffectiveURL);
+			}
+			SetupSegmentDownloadStatsFromConnectionInfo(CurrentConnectionInfo);
+			if (!bSuccess)
+			{
+				FString ErrMsg = BoxLocator.GetErrorMessage();
+				if (ErrMsg.IsEmpty())
+				{
+					ErrMsg = CurrentConnectionInfo.StatusInfo.ErrorDetail.GetMessage();
+				}
+				SetError(FString::Printf(TEXT("Self-initializing media 'moov' scan error: %s"),	*ErrMsg), INTERNAL_SEG_ERROR_INIT_SEGMENT_DOWNLOAD_ERROR);
+				return false;
+			}
+			const UtilsMP4::FMP4RootBoxLocator::FBoxInfo* Moov = Boxes.FindByPredicate([InType=MP4Utilities::MakeBoxAtom('m','o','o','v')](const UtilsMP4::FMP4RootBoxLocator::FBoxInfo& InBox){return InBox.Type == InType;});
+			if (Moov)
+			{
+				StaticSegmentDataReader.SetSourceBuffer(Moov->DataBuffer);
+				ContainerFormat = FStreamSegmentRequestCommon::EContainerFormat::ISO_14496_12;
+			}
+			else
+			{
+				SetError(FString::Printf(TEXT("Self-initializing media does not have a 'moov' box")), INTERNAL_SEG_ERROR_INIT_SEGMENT_NOTFOUND_ERROR);
+				return false;
+			}
 		}
 	}
 
@@ -1200,23 +1295,19 @@
 	{
 		case FStreamSegmentRequestCommon::EContainerFormat::ISO_14496_12:
 		{
-			TSharedPtrTS<IParserISO14496_12> MP4InitSeg = IParserISO14496_12::CreateParser();
-			FMediaSegmentBoxCallback ParseBoxCallback;
-			parseError = MP4InitSeg->ParseHeader(&StaticSegmentDataReader, &ParseBoxCallback, PlayerSessionService, nullptr);
-			if (parseError == UEMEDIA_ERROR_OK || parseError == UEMEDIA_ERROR_END_OF_STREAM)
+			if (!Stream::CVarElectraStreamUseLegacyMP4.GetValueOnAnyThread())
 			{
-				// Parse the tracks of the init segment. We do this mainly to get to the CSD we might need should we have to insert filler data later.
-				parseError = MP4InitSeg->PrepareTracks(PlayerSessionService, TSharedPtrTS<const IParserISO14496_12>());
-				if (parseError == UEMEDIA_ERROR_OK)
+				TSharedPtr<FParsedMOOVCommon, ESPMode::ThreadSafe> ParsedMOOV = MakeShared<FParsedMOOVCommon, ESPMode::ThreadSafe>();
+				if (ParsedMOOV->Parse(StaticSegmentDataReader.GetConstArrayView()))
 				{
-					InitSegmentData.Emplace<TSharedPtrTS<const IParserISO14496_12>>(MP4InitSeg);
+					InitSegmentData.Emplace<TSharedPtrTS<FParsedMOOVCommon>>(ParsedMOOV);
 					// Add this to the entity cache in case it needs to be retrieved again.
 					if (EntityCache.IsValid())
 					{
 						IPlayerEntityCache::FCacheItem CacheItem;
 						CacheItem.URL = InitURLToUse.Url.URL;
 						CacheItem.Range = InitURLToUse.Range;
-						CacheItem.Parsed14496_12Data = MP4InitSeg;
+						CacheItem.ParsedMOOV = MoveTemp(ParsedMOOV);
 						EntityCache->CacheEntity(CacheItem);
 					}
 					DownloadStats.bWasSuccessful = true;
@@ -1227,16 +1318,50 @@
 				{
 					DownloadStats.bParseFailure = true;
 					StreamSelector->ReportDownloadEnd(DownloadStats);
-					SetError(FString::Printf(TEXT("Track preparation of init segment \"%s\" failed"), *InitURLToUse.Url.URL), INTERNAL_SEG_ERROR_INIT_SEGMENT_PARSE_ERROR);
+					SetError(FString::Printf(TEXT("Parse error of init segment \"%s\""), *InitURLToUse.Url.URL), INTERNAL_SEG_ERROR_INIT_SEGMENT_PARSE_ERROR);
 					return false;
 				}
 			}
 			else
 			{
-				DownloadStats.bParseFailure = true;
-				StreamSelector->ReportDownloadEnd(DownloadStats);
-				SetError(FString::Printf(TEXT("Parse error of init segment \"%s\""), *InitURLToUse.Url.URL), INTERNAL_SEG_ERROR_INIT_SEGMENT_PARSE_ERROR);
-				return false;
+				TSharedPtrTS<IParserISO14496_12> MP4InitSeg = IParserISO14496_12::CreateParser();
+				FMediaSegmentBoxCallback ParseBoxCallback;
+				parseError = MP4InitSeg->ParseHeader(&StaticSegmentDataReader, &ParseBoxCallback, PlayerSessionService, nullptr);
+				if (parseError == UEMEDIA_ERROR_OK || parseError == UEMEDIA_ERROR_END_OF_STREAM)
+				{
+					// Parse the tracks of the init segment. We do this mainly to get to the CSD we might need should we have to insert filler data later.
+					parseError = MP4InitSeg->PrepareTracks(PlayerSessionService, TSharedPtrTS<const IParserISO14496_12>());
+					if (parseError == UEMEDIA_ERROR_OK)
+					{
+						InitSegmentData.Emplace<TSharedPtrTS<const IParserISO14496_12>>(MP4InitSeg);
+						// Add this to the entity cache in case it needs to be retrieved again.
+						if (EntityCache.IsValid())
+						{
+							IPlayerEntityCache::FCacheItem CacheItem;
+							CacheItem.URL = InitURLToUse.Url.URL;
+							CacheItem.Range = InitURLToUse.Range;
+							CacheItem.Parsed14496_12Data = MP4InitSeg;
+							EntityCache->CacheEntity(CacheItem);
+						}
+						DownloadStats.bWasSuccessful = true;
+						StreamSelector->ReportDownloadEnd(DownloadStats);
+						return true;
+					}
+					else
+					{
+						DownloadStats.bParseFailure = true;
+						StreamSelector->ReportDownloadEnd(DownloadStats);
+						SetError(FString::Printf(TEXT("Track preparation of init segment \"%s\" failed"), *InitURLToUse.Url.URL), INTERNAL_SEG_ERROR_INIT_SEGMENT_PARSE_ERROR);
+						return false;
+					}
+				}
+				else
+				{
+					DownloadStats.bParseFailure = true;
+					StreamSelector->ReportDownloadEnd(DownloadStats);
+					SetError(FString::Printf(TEXT("Parse error of init segment \"%s\""), *InitURLToUse.Url.URL), INTERNAL_SEG_ERROR_INIT_SEGMENT_PARSE_ERROR);
+					return false;
+				}
 			}
 		}
 		case FStreamSegmentRequestCommon::EContainerFormat::Matroska_WebM:
@@ -1666,4 +1791,695 @@
 FStreamSegmentReaderCommon::FStreamHandler::EHandleResult FStreamSegmentReaderCommon::FStreamHandler::HandleMP4Media(TSharedPtrTS<FStreamSegmentRequestCommon> InRequest)
 {
 	// We need to have an mp4 init segment.
+	if (!InitSegmentData.IsType<TSharedPtrTS<FParsedMOOVCommon>>())
+	{
+		return SetError(FString::Printf(TEXT("Wrong init segment type for fmp4 media segment")), INTERNAL_SEG_ERROR_BAD_SEGMENT_TYPE);
+	}
+	TSharedPtrTS<FParsedMOOVCommon> InitSegment = InitSegmentData.Get<TSharedPtrTS<FParsedMOOVCommon>>();
+	if (!InitSegment.IsValid())
+	{
+		return SetError(FString::Printf(TEXT("No init segment for fmp4 media segment")), INTERNAL_SEG_ERROR_BAD_SEGMENT_TYPE);
+	}
+	if (InitSegment->GetTrackMap().IsEmpty())
+	{
+		return SetError(FString::Printf(TEXT("Init segment contains no tracks")), INTERNAL_SEG_ERROR_BAD_SEGMENT_TYPE);
+	}
+	TSortedMap<uint64, TSharedPtrTS<FActiveTrackData>> NewTrackDataTypeMap[4];
+	for(auto &tkmIt : InitSegment->GetTrackMap())
+	{
+		const FParsedMOOVCommon::FTrackInfo& tkInf(tkmIt.Value);
+		if (tkInf.CodecFormat.Type == FCodecTypeFormat::EType::Video && InRequest->bIgnoreVideo)
+		{
+			continue;
+		}
+		else if (tkInf.CodecFormat.Type == FCodecTypeFormat::EType::Audio && InRequest->bIgnoreAudio)
+		{
+			continue;
+		}
+		else if (tkInf.CodecFormat.Type == FCodecTypeFormat::EType::Subtitle && InRequest->bIgnoreSubtitles)
+		{
+			continue;
+		}
+
+		TSharedPtrTS<FActiveTrackData> td = MakeSharedTS<FActiveTrackData>();
+		td->StreamType = tkInf.CodecFormat.Type == FCodecTypeFormat::EType::Video ? EStreamType::Video :
+						 tkInf.CodecFormat.Type == FCodecTypeFormat::EType::Audio ? EStreamType::Audio : EStreamType::Subtitle;
+		const int32 stIdx = StreamTypeToArrayIndex(td->StreamType);
+		// Copy the source buffer info into a new instance and set the playback sequence ID in it.
+		check(InRequest->SourceBufferInfo[stIdx].IsValid());
+		if (!InRequest->SourceBufferInfo[stIdx].IsValid())
+		{
+			continue;
+		}
+		td->BufferSourceInfo = MakeSharedTS<FBufferSourceInfo>(*(InRequest->SourceBufferInfo[stIdx]));
+		td->BufferSourceInfo->PlaybackSequenceID = InRequest->GetPlaybackSequenceID();
+		// Set the CSD
+		td->CSD = MakeSharedTS<FAccessUnit::CodecData>();
+		td->CSD->CodecSpecificData = tkInf.CodecFormat.CSD;
+		td->CSD->RawCSD = tkInf.CodecFormat.DCR;
+		td->CSD->ParsedInfo = FStreamCodecInformation(tkInf.CodecFormat);
+		// Set information from the playlist codec information that may not be available or accurate in the init segment.
+		td->CSD->ParsedInfo.SetBitrate(InRequest->CodecInfo[stIdx].GetBitrate());
+		NewTrackDataTypeMap[stIdx].Emplace(tkInf.TrackID, MoveTemp(td));
+	}
+	// At present we only want to have a single track per type in the media segment. If there are more we use the one
+	// with the smallest ID (a stipulation made by the HLS specification; we apply it regardless of protocol)
+	for(int32 nt=0; nt<3; ++nt)
+	{
+		if (NewTrackDataTypeMap[nt].Num())
+		{
+			auto it(NewTrackDataTypeMap[nt].CreateConstIterator());
+			TrackDataMap.Add(it.Key(), it.Value());
+		}
+	}
+	if (TrackDataMap.Num() == 0 || (InRequest->bIsMultiplex && NewTrackDataTypeMap[StreamTypeToArrayIndex(InRequest->GetType())].Num() == 0))
+	{
+		SetError(FString::Printf(TEXT("Init segment contains no usable tracks")), INTERNAL_SEG_ERROR_BAD_SEGMENT_TYPE);
+		return EHandleResult::Failed;
+	}
+
+	SelectPrimaryTrackData(InRequest);
+
+	// If this is a filler segment then we are done at this point and let the end-of-segment handling
+	// create the filler data for the active tracks we just set up.
+	if (InRequest->bInsertFillerData)
+	{
+		return EHandleResult::Finished;
+	}
+
+	// See if the segment is encrypted (sample encryption, not whole segment)
+	// The assumption is that all tracks in the segment are encrypted the same way and a single decrypter will do.
+	if (!ReadBuffer.bIsEncrypted && InRequest->DrmMedia.DrmClient.IsValid())
+	{
+		check(!Decrypter.IsValid());
+		if (InRequest->DrmMedia.DrmClient->CreateDecrypter(Decrypter, InRequest->DrmMedia.DrmMimeType) != ElectraCDM::ECDMError::Success)
+		{
+			SetError(FString::Printf(TEXT("Failed to create decrypter for segment, \"%s\""), *InRequest->DrmMedia.DrmClient->GetLastErrorMessage()), INTERNAL_SEG_ERROR_FAILED_TO_DECRYPT);
+			return EHandleResult::Failed;
+		}
+	}
+
+	// Create a box loader wrapper.
+	TSharedPtr<FBoxDataLoader, ESPMode::ThreadSafe> BoxLoader = MakeShared<FBoxDataLoader, ESPMode::ThreadSafe>(this);
+
+	// Enter the parsing loop.
+	bool bDone = false;
+	FTimeValue SegmentDuration(InRequest->Segment.Duration, InRequest->Segment.Timescale, 0);
+	FTimeValue TimeOffset = InRequest->PeriodStart + InRequest->AST + InRequest->AdditionalAdjustmentTime;
+	bool bDoNotTruncateAtPresentationEnd = PlayerSessionService->GetOptionValue(OptionKeyDoNotTruncateAtPresentationEnd).SafeGetBool(false);
+	bool bIsLastSegment = InRequest->Segment.bIsLastInPeriod;
+	const FString& RequestURL(InRequest->Segment.MediaURL.Url.URL);
+	bool bSkippedBecauseOfTimestampCheck = false;
+	int32 StreamTypeAUCount[4] { 0,0,0,0 };
+	int64 LastByteOffsetOfMDAT = -1;
+	bool bTimeCheckPassed = InRequest->TimestampVars.Next.bCheck == false;
+	bool bSendMetaData = !!InitSegment->GetMETABox();
+	bool bSendEventData = true;
+	while(!bDone && !HasErrored() && !HasReadBeenAborted())
+	{
+		if (HasReachedEOF())
+		{
+			bDone = true;
+			break;
+		}
+		Metrics::FSegmentDownloadStats::FMovieChunkInfo MoofInfo;
+		MoofInfo.HeaderOffset = GetCurrentOffset();
+		LastByteOffsetOfMDAT = -1;
+
+		// Read all the root boxes up to the next `mdat` box.
+		MP4Utilities::FMP4BoxLocator BoxLocator;
+		const TArray<uint32> UnspecBoxes;
+		TArray<TSharedPtr<MP4Utilities::FMP4BoxData, ESPMode::ThreadSafe>> RootBoxes;
+		bool bGotBoxes = BoxLocator.LocateAndReadRootBoxes(RootBoxes, BoxLoader, UnspecBoxes, UnspecBoxes, true, UnspecBoxes, FMP4DataLoader::FCancellationCheckDelegate());
+		// If we did not get the necessary boxes see if we hit EOF. This may happen if there are trailing boxes that we do
+		// not need, like `mfra` or similar.
+		if (!bGotBoxes)
+		{
+			bDone = true;
+			if (HasReachedEOF())
+			{
+				break;
+			}
+			else if (!HasErrored() && !HasReadBeenAborted())
+			{
+				SetError(FString::Printf(TEXT("Failed to find movie boxes in segment \"%s\""), *RequestURL), INTERNAL_SEG_ERROR_BAD_MEDIA_SEGMENT);
+			}
+			break;
+		}
+
+		TArray<TSharedPtr<MP4Boxes::FMP4BoxBase, ESPMode::ThreadSafe>> ParsedRootBoxes;
+		ParsedRootBoxes.SetNum(RootBoxes.Num());
+		TSharedPtr<MP4Boxes::FMP4BoxMETA, ESPMode::ThreadSafe> MetaBox;
+		TSharedPtr<MP4Boxes::FMP4BoxSIDX, ESPMode::ThreadSafe> SidxBox;
+		TArray<TSharedPtr<MP4Boxes::FMP4BoxPRFT, ESPMode::ThreadSafe>> ProducerReferenceTimeBoxes;
+		TArray<TSharedPtr<MP4Boxes::FMP4BoxEMSG, ESPMode::ThreadSafe>> EventMessageBoxes;
+		TArray<TSharedPtr<MP4Boxes::FMP4BoxPSSH, ESPMode::ThreadSafe>> ProtectionSystemBoxes;
+		for(int32 i=0; i<RootBoxes.Num(); ++i)
+		{
+			MP4Boxes::FMP4BoxTreeParser tp;
+			if (tp.ParseBoxTree(RootBoxes[i]))
+			{
+				ParsedRootBoxes[i] = tp.GetBoxTree();
+				switch(ParsedRootBoxes[i]->GetType())
+				{
+					case MP4Utilities::MakeBoxAtom('m','d','a','t'):
+					{
+						LastByteOffsetOfMDAT = ParsedRootBoxes[i]->GetBoxFileOffset() + ParsedRootBoxes[i]->GetBoxSize();
+						break;
+					}
+					case MP4Utilities::MakeBoxAtom('s','t','y','p'):
+					case MP4Utilities::MakeBoxAtom('f','t','y','p'):
+					{
+						auto XtypBox = StaticCastSharedPtr<MP4Boxes::FMP4BoxFTYP>(ParsedRootBoxes[i]);
+						if (XtypBox->GetAllBrands().Contains(MP4Utilities::MakeBoxAtom('l','m','s','g')))
+						{
+							InRequest->Segment.bSawLMSG = true;
+						}
+						break;
+					}
+					case MP4Utilities::MakeBoxAtom('u','d','t','a'):
+					{
+						MetaBox = StaticCastSharedPtr<MP4Boxes::FMP4BoxUDTA>(ParsedRootBoxes[i])->FindBoxRecursive<MP4Boxes::FMP4BoxMETA>(MP4Utilities::MakeBoxAtom('m','e','t','a'), 0);
+						bSendMetaData = true;
+						break;
+					}
+					case MP4Utilities::MakeBoxAtom('s','i','d','x'):
+					{
+						SidxBox = StaticCastSharedPtr<MP4Boxes::FMP4BoxSIDX>(ParsedRootBoxes[i]);
+						break;
+					}
+					case MP4Utilities::MakeBoxAtom('p','r','f','t'):
+					{
+						ProducerReferenceTimeBoxes.Add(StaticCastSharedPtr<MP4Boxes::FMP4BoxPRFT>(ParsedRootBoxes[i]));
+						break;
+					}
+					case MP4Utilities::MakeBoxAtom('e','m','s','g'):
+					{
+						EventMessageBoxes.Add(StaticCastSharedPtr<MP4Boxes::FMP4BoxEMSG>(ParsedRootBoxes[i]));
+						break;
+					}
+					case MP4Utilities::MakeBoxAtom('p','s','s','h'):
+					{
+						ProtectionSystemBoxes.Add(StaticCastSharedPtr<MP4Boxes::FMP4BoxPSSH>(ParsedRootBoxes[i]));
+						break;
+					}
+					default:
+					{
+						break;
+					}
+				}
+			}
+			else
+			{
+				return SetError(FString::Printf(TEXT("Failed to parse movie boxes in segment \"%s\""), *RequestURL), INTERNAL_SEG_ERROR_BAD_MEDIA_SEGMENT);
+			}
+		}
+		if (LastByteOffsetOfMDAT < 0)
+		{
+			return SetError(FString::Printf(TEXT("`mdat` box not found in segment \"%s\""), *RequestURL), INTERNAL_SEG_ERROR_BAD_MEDIA_SEGMENT);
+		}
+
+		struct FSelectedTrackParams
+		{
+			FDRMTypeFormat DrmFormat;
+			ElectraCDM::FMediaCDMSampleInfo SampleEncryptionInfo;
+			TSharedPtr<MP4Boxes::FMP4Track, ESPMode::ThreadSafe> Track;
+			TSharedPtr<MP4Boxes::FMP4Track::FFragmentInfo, ESPMode::ThreadSafe> FragmentInfo;
+			TSharedPtr<MP4Boxes::FMP4Track::FIterator, ESPMode::ThreadSafe> Iterator;
+			TArray<TSharedPtr<MP4Boxes::FMP4BoxPRFT, ESPMode::ThreadSafe>> PRFTBoxes;
+			uint32 Timescale = 0;
+			bool bIsEncrypted = false;
+			bool SetupBaseEncryption()
+			{
+				bIsEncrypted = DrmFormat.EncryptionInfo.IsType<FDRMTypeFormat::FISOEncryptionInfo>();
+				if (bIsEncrypted)
+				{
+					const FDRMTypeFormat::FISOEncryptionInfo& encInf(DrmFormat.EncryptionInfo.Get<FDRMTypeFormat::FISOEncryptionInfo>());
+					SampleEncryptionInfo.DefaultKID = encInf.DefaultKID;
+					SampleEncryptionInfo.IV = encInf.DefaultIV;
+					SampleEncryptionInfo.Scheme4CC = encInf.Scheme;
+					if (encInf.BlockPattern.IsSet())
+					{
+						SampleEncryptionInfo.Pattern.PatternType = 1;
+						SampleEncryptionInfo.Pattern.CryptByteBlock = encInf.BlockPattern.GetValue().CryptByteBlock;
+						SampleEncryptionInfo.Pattern.SkipByteBlock = encInf.BlockPattern.GetValue().SkipByteBlock;
+					}
+					return true;
+				}
+				return false;
+			}
+			bool SetupSampleEncryption()
+			{
+				if (bIsEncrypted)
+				{
+					const MP4Boxes::FMP4BoxSENC::FEntry& enc = Iterator->GetEncryptionInfo();
+					SampleEncryptionInfo.SubSamples.Empty();
+					SampleEncryptionInfo.IV = enc.IV;
+					for(auto &sub : enc.SubSamples)
+					{
+						SampleEncryptionInfo.SubSamples.Add({.NumClearBytes = sub.NumClearBytes, .NumEncryptedBytes = sub.NumEncryptedBytes});
+					}
+					return true;
+				}
+				return false;
+			}
+		};
+		TMap<uint32, FSelectedTrackParams> SelectedTrackParamMap;
+
+		// Prepare tracks
+		for(auto &tkmIt : InitSegment->GetTrackMap())
+		{
+			const FParsedMOOVCommon::FTrackInfo& tkInf(tkmIt.Value);
+			FSelectedTrackParams SelectedTrackParam;
+			SelectedTrackParam.Timescale = tkInf.TrackTimescale;
+			SelectedTrackParam.DrmFormat = tkInf.DrmFormat;
+			SelectedTrackParam.SetupBaseEncryption();
+			SelectedTrackParam.FragmentInfo = MakeShared<MP4Boxes::FMP4Track::FFragmentInfo, ESPMode::ThreadSafe>();
+			if (!SelectedTrackParam.FragmentInfo->Prepare(tkInf.TrakBox, InitSegment->GetMVEXBox(), ParsedRootBoxes))
+			{
+				return SetError(FString::Printf(TEXT("Failed to prepare fragment information in segment \"%s\""), *RequestURL), INTERNAL_SEG_ERROR_BAD_MEDIA_SEGMENT);
+			}
+			SelectedTrackParam.Track = MP4Boxes::FMP4Track::Create(tkInf.TrakBox, SelectedTrackParam.FragmentInfo);
+			MP4Utilities::FFractionalTime Inf;
+			if (!SelectedTrackParam.Track->Prepare(Inf, Inf))
+			{
+				return SetError(FString::Printf(TEXT("Failed to prepare track in segment \"%s\""), *RequestURL), INTERNAL_SEG_ERROR_BAD_MEDIA_SEGMENT);
+			}
+			SelectedTrackParam.Iterator = SelectedTrackParam.Track->CreateIterator();
+			if (!SelectedTrackParam.Iterator)
+			{
+				return SetError(FString::Printf(TEXT("Failed to prepare track iterator in segment \"%s\""), *RequestURL), INTERNAL_SEG_ERROR_BAD_MEDIA_SEGMENT);
+			}
+			// Collect the `prft` boxes for this track
+			for(auto &prft : ProducerReferenceTimeBoxes)
+			{
+				if (prft->GetReferenceTrackID() == tkInf.TrackID)
+				{
+					SelectedTrackParam.PRFTBoxes.Add(prft);
+				}
+			}
+			SelectedTrackParamMap.Add(tkInf.TrackID, SelectedTrackParam);
+		}
+
+		// Send the metadata from the init segment or from this `moof`.
+		// This is not per-track metadata, only segment-global metadata!
+		if (bSendMetaData)
+		{
+			HandleMP4Metadata(InRequest, InitSegment->GetMETABox(), MetaBox);
+			bSendMetaData = false;
+		}
+
+		// Send event data from `emsg` boxes. Do this only for the first round of boxes we encounter, not repeated ones.
+		if (bSendEventData && EventMessageBoxes.Num())
+		{
+			HandleMP4EventMessages(InRequest, SidxBox, EventMessageBoxes);
+			bSendEventData = false;
+		}
+
+		// Locate the track iterator that will handle the next incoming data.
+		auto GetNextTrackIndexToIterate = [&]() -> int64
+		{
+			int64 NextIteratorFileDataOffset = TNumericLimits<int64>::Max();
+			int64 ItIdx = -1;
+			for(auto &stpIt : SelectedTrackParamMap)
+			{
+				if (stpIt.Value.Iterator && stpIt.Value.Iterator->GetSampleFileOffset() < NextIteratorFileDataOffset)
+				{
+					NextIteratorFileDataOffset = stpIt.Value.Iterator->GetSampleFileOffset();
+					ItIdx = stpIt.Key;
+				}
+			}
+			return ItIdx;
+		};
+
+
+		// Process consecutive samples from each track.
+		int64 NextTrackIndex = -1;
+		while(!bDone && !HasErrored() && !HasReadBeenAborted())
+		{
+			if (NextTrackIndex < 0)
+			{
+				NextTrackIndex = GetNextTrackIndexToIterate();
+				// All tracks finished in this `moof`?
+				if (NextTrackIndex < 0)
+				{
+					break;
+				}
+			}
+
+			FSelectedTrackParams& stp(SelectedTrackParamMap[(uint32)NextTrackIndex]);
+
+			MP4Boxes::FMP4Track::FIterator::FChunkInfo SampleChunkInfo;
+			if (!stp.Iterator->GetCurrentChunkRemainingSampleInfo(SampleChunkInfo))
+			{
+				return SetError(FString::Printf(TEXT("Failed to get sample chunk information in segment \"%s\""), *RequestURL), INTERNAL_SEG_ERROR_BAD_MEDIA_SEGMENT);
+			}
+
+			// Set this as the currently active track.
+			CurrentlyActiveTrackData = TrackDataMap[NextTrackIndex];
+			const int32 stIdx = StreamTypeToArrayIndex(CurrentlyActiveTrackData->StreamType);
+			const uint32 TrackTimescale = stp.Timescale;
+
+			// Perform some checks and adjustments for the first access unit.
+			// Note: this is not necessarily correct in that the times given in the segment are
+			//       presently in media track local time and since we may be dealing with multiple
+			//       tracks here this one value does not necessarily apply to all.
+			if (StreamTypeAUCount[stIdx] == 0)
+			{
+				CurrentlyActiveTrackData->MediaLocalFirstAUTime = InRequest->Segment.MediaLocalFirstAUTime;
+				CurrentlyActiveTrackData->MediaLocalLastAUTime = bDoNotTruncateAtPresentationEnd ? TNumericLimits<int64>::Max() : InRequest->Segment.MediaLocalLastAUTime;
+				CurrentlyActiveTrackData->PTO = InRequest->Segment.PTO;
+
+				if (TrackTimescale != InRequest->Segment.Timescale)
+				{
+					// Need to rescale the AU times from the MPD timescale to the media timescale.
+					CurrentlyActiveTrackData->MediaLocalFirstAUTime = FTimeFraction(InRequest->Segment.MediaLocalFirstAUTime, InRequest->Segment.Timescale).GetAsTimebase(TrackTimescale);
+					CurrentlyActiveTrackData->MediaLocalLastAUTime = CurrentlyActiveTrackData->MediaLocalLastAUTime == TNumericLimits<int64>::Max() ? CurrentlyActiveTrackData->MediaLocalLastAUTime : FTimeFraction(InRequest->Segment.MediaLocalLastAUTime, InRequest->Segment.Timescale).GetAsTimebase(TrackTimescale);
+					CurrentlyActiveTrackData->PTO = FTimeFraction(InRequest->Segment.PTO, InRequest->Segment.Timescale).GetAsTimebase(TrackTimescale);
+					/*
+					Removed because this keeps irritating people who do not understand the problem and how to fix this in their MPD generator.
+
+						if (InRequest->StreamingProtocol == FStreamSegmentRequestCommon::EStreamingProtocol::DASH && !InRequest->bWarnedAboutTimescale)
+						{
+							InRequest->bWarnedAboutTimescale = true;
+							LogMessage(IInfoLog::ELevel::Verbose, FString::Printf(TEXT("Track timescale %u differs from timescale of %u in MPD or segment index. This may cause playback problems!"), TrackTimescale, InRequest->Segment.Timescale));
+						}
+					*/
+				}
+				// Set the PTO in the codec data extras. This is a rarely used value and constant for the segment.
+				// The only use we have for it so far is to remap subtitle timestamps to split periods.
+				CurrentlyActiveTrackData->CSD->ParsedInfo.GetExtras().Set(StreamCodecInformationOptions::PresentationTimeOffset, FVariantValue(FTimeValue(CurrentlyActiveTrackData->PTO, TrackTimescale)));
+
+				// Producer reference time (DASH only)
+				if (InRequest->Segment.MeasureLatencyViaReferenceTimeInfoID >= 0 && InRequest->Segment.ProducerReferenceTimeInfos.Num())
+				{
+					// We only look at inband 'prtf' boxes if the <ProducerReferenceTime> element in the MPD tells us to.
+					// This is similar to events that are not to be considered any more if the MPD doesn't specify them.
+					for(auto &MPDPrtf : InRequest->Segment.ProducerReferenceTimeInfos)
+					{
+						if (MPDPrtf.GetID() == (uint32)InRequest->Segment.MeasureLatencyViaReferenceTimeInfoID)
+						{
+							// Use the inband 'prft' boxes?
+							if (MPDPrtf.bInband)
+							{
+								bool bFound = false;
+								for(auto &MP4Prtf : stp.PRFTBoxes)
+								{
+									if ((MPDPrtf.Type == FProducerReferenceTimeInfo::EType::Encoder && MP4Prtf->GetReferenceType() == MP4Boxes::FMP4BoxPRFT::EReferenceType::Encoder) ||
+										(MPDPrtf.Type == FProducerReferenceTimeInfo::EType::Captured && MP4Prtf->GetReferenceType() == MP4Boxes::FMP4BoxPRFT::EReferenceType::Captured))
+									{
+										RFC5905::ParseNTPTime(CurrentlyActiveTrackData->ProducerTime.Base, MP4Prtf->GetNTPTimestamp());
+										CurrentlyActiveTrackData->ProducerTime.Media = MP4Prtf->GetMediaTime();
+										bFound = true;
+										break;
+									}
+								}
+								// When the MPD says that there are inband prtf's then this has to be so. If for some reason this is not the case
+								// then what are we to do?
+								if (!bFound)
+								{
+									// We take the values from the MPD here, which may be better than nothing?!
+									CurrentlyActiveTrackData->ProducerTime.Base = MPDPrtf.WallclockTime;
+									CurrentlyActiveTrackData->ProducerTime.Media = MPDPrtf.PresentationTime;
+								}
+							}
+							else
+							{
+								// Use values from MPD
+								CurrentlyActiveTrackData->ProducerTime.Base = MPDPrtf.WallclockTime;
+								CurrentlyActiveTrackData->ProducerTime.Media = MPDPrtf.PresentationTime;
+							}
+							break;
+						}
+					}
+				}
+			}
+
+			// Establish time mapping for this track
+			if (!CurrentlyActiveTrackData->TimeMappingOffset.IsValid())
+			{
+				if (InRequest->StreamingProtocol == FStreamSegmentRequestCommon::EStreamingProtocol::HLS && InRequest->HLS.bNoPDTMapping == false)
+				{
+					const int64 bmdt = stp.Iterator->GetBaseMediaDecodeTime().Get(MP4Utilities::FFractionalTime(0,1)).GetNumerator();
+					FTimeValue BaseMediaDecodeTime(bmdt, TrackTimescale, 0);
+					CurrentlyActiveTrackData->TimeMappingOffset = FTimeValue().SetFromHNS(InRequest->Segment.Time) - BaseMediaDecodeTime;
+				}
+				else
+				{
+					CurrentlyActiveTrackData->TimeMappingOffset.SetToZero();
+				}
+			}
+
+
+			// Process all samples in the current chunk
+			while(NextTrackIndex >= 0 && !bDone && !HasErrored() && !HasReadBeenAborted())
+			{
+				// Get the DTS and PTS. Those are 0-based in a fragment and offset by the base media decode time of the fragment.
+				MP4Utilities::FFractionalTime AUDTS = stp.Iterator->GetEffectiveDTS();
+				MP4Utilities::FFractionalTime AUPTS = stp.Iterator->GetEffectivePTS();
+				MP4Utilities::FFractionalTime AUDuration = stp.Iterator->GetDuration();
+
+				// Create access unit
+				FAccessUnit *AccessUnit = FAccessUnit::Create(Parameters.MemoryProvider);
+				AccessUnit->ESType = CurrentlyActiveTrackData->StreamType;
+				AccessUnit->AUSize = (uint32) stp.Iterator->GetSampleSize();
+				AccessUnit->AUData = AccessUnit->AllocatePayloadBuffer(AccessUnit->AUSize);
+				AccessUnit->bIsFirstInSequence = CurrentlyActiveTrackData->bIsFirstInSequence;
+				AccessUnit->bIsSyncSample = stp.Iterator->IsSyncOrRAPSample();
+				AccessUnit->bIsDummyData = false;
+				AccessUnit->AUCodecData = CurrentlyActiveTrackData->CSD;
+				AccessUnit->BufferSourceInfo = CurrentlyActiveTrackData->BufferSourceInfo;
+
+				FTimeValue Duration(AUDuration.GetNumerator(), AUDuration.GetDenominator(), 0);
+				AccessUnit->Duration = Duration;
+
+				// Offset the AU's DTS and PTS to the time mapping of the segment.
+				AccessUnit->DTS.SetFromND(AUDTS.GetNumerator() - CurrentlyActiveTrackData->PTO, TrackTimescale);
+				AccessUnit->DTS += TimeOffset + CurrentlyActiveTrackData->TimeMappingOffset;
+				AccessUnit->PTS.SetFromND(AUPTS.GetNumerator() - CurrentlyActiveTrackData->PTO, TrackTimescale);
+				AccessUnit->PTS += TimeOffset + CurrentlyActiveTrackData->TimeMappingOffset;
+
+				// Remember the first AU's DTS for this stream type in the segment
+				if (!InRequest->TimestampVars.Local.First[stIdx].IsValid())
+				{
+					InRequest->TimestampVars.Local.First[stIdx] = AccessUnit->DTS;
+				}
+				if (CurrentlyActiveTrackData == PrimaryTrackData)
+				{
+					if (StreamTypeAUCount[stIdx] == 0)
+					{
+						// Check that the timestamps are greater than what we are expecting.
+						// This is to prevent reading the same segment data from a different stream again after a stream switch.
+						if (InRequest->TimestampVars.Next.bCheck)
+						{
+							check(InRequest->TimestampVars.Next.ExpectedLargerThan[stIdx].IsValid());
+							if (InRequest->TimestampVars.Local.First[stIdx] <= InRequest->TimestampVars.Next.ExpectedLargerThan[stIdx])
+							{
+								FAccessUnit::Release(AccessUnit);
+								AccessUnit = nullptr;
+								bDone = true;
+								InRequest->TimestampVars.Next.bFailed = true;
+								bSkippedBecauseOfTimestampCheck = true;
+								break;
+							}
+						}
+						bTimeCheckPassed = true;
+
+						InRequest->FirstTimestampReceivedDelegate.ExecuteIfBound(InRequest);
+					}
+
+					// If the track uses encryption we update the DRM system with the PSSH boxes that are currently in use.
+					// NOTE: This works only if there is a single decrypter for any of the multiplexed tracks.
+					//       If there is ever more than one this needs to be extended.
+					if (Decrypter.IsValid() && ProtectionSystemBoxes.Num())
+					{
+						TArray<TArray<uint8>> PSSHBoxes;
+						for(auto &psb : ProtectionSystemBoxes)
+						{
+							PSSHBoxes.Emplace(psb->GetBoxDataRAW());
+						}
+						Decrypter->UpdateInitDataFromMultiplePSSH(PSSHBoxes);
+					}
+				}
+
+				// Set the sequence index member and update all timestamps with it as well.
+				AccessUnit->SequenceIndex = InRequest->TimestampSequenceIndex;
+				AccessUnit->DTS.SetSequenceIndex(InRequest->TimestampSequenceIndex);
+				AccessUnit->PTS.SetSequenceIndex(InRequest->TimestampSequenceIndex);
+
+				AccessUnit->EarliestPTS.SetFromND(CurrentlyActiveTrackData->MediaLocalFirstAUTime - CurrentlyActiveTrackData->PTO, TrackTimescale);
+				AccessUnit->EarliestPTS += TimeOffset;
+				if (InRequest->Segment.bFrameAccuracyRequired && InRequest->FrameAccurateStartTime.IsValid())
+				{
+					AccessUnit->EarliestPTS = InRequest->FrameAccurateStartTime;
+				}
+				AccessUnit->EarliestPTS.SetSequenceIndex(InRequest->TimestampSequenceIndex);
+				if (CurrentlyActiveTrackData->MediaLocalLastAUTime != TNumericLimits<int64>::Max())
+				{
+					AccessUnit->LatestPTS.SetFromND(CurrentlyActiveTrackData->MediaLocalLastAUTime - CurrentlyActiveTrackData->PTO, TrackTimescale);
+					AccessUnit->LatestPTS += TimeOffset;
+				}
+				else
+				{
+					AccessUnit->LatestPTS.SetToPositiveInfinity();
+				}
+				AccessUnit->LatestPTS.SetSequenceIndex(InRequest->TimestampSequenceIndex);
+
+				AccessUnit->ProducerReferenceTime = CurrentlyActiveTrackData->ProducerTime.Base + FTimeValue(AUDTS.GetNumerator() - CurrentlyActiveTrackData->ProducerTime.Media, TrackTimescale);
+
+				// Set up the sample encryption if any. This needs to be done prior to reading the sample data.
+				bool bIsSampleEncrypted = stp.SetupSampleEncryption();
+				if (bIsSampleEncrypted)
+				{
+					// If the sample is encrypted in a HLS stream we need to replace the default KID.
+					// That's because when the .m3u8 playlist was parsed there is no KID in the `#EXT-X-KEY` tag (for KEYFORMAT identity)
+					// and we had to create a "KID" by hashing the license key URL, and thus need to now set that hash as the KID.
+					if (InRequest->StreamingProtocol == FStreamSegmentRequestCommon::EStreamingProtocol::HLS)
+					{
+						stp.SampleEncryptionInfo.DefaultKID = InRequest->DrmMedia.DrmKID;
+					}
+				}
+
+				const int64 AUFileOffset = stp.Iterator->GetSampleFileOffset();
+				const int64 NumBytesToSkip = AUFileOffset - GetCurrentOffset();
+				if (NumBytesToSkip < 0)
+				{
+					// Current read position is already farther than where the data is supposed to be.
+					FAccessUnit::Release(AccessUnit);
+					AccessUnit = nullptr;
+					SetError(FString::Printf(TEXT("Read position already at %lld but data starts at %lld in segment \"%s\""), (long long int)GetCurrentOffset(), (long long int)AUFileOffset, *RequestURL), INTERNAL_SEG_ERROR_BAD_MEDIA_SEGMENT);
+					break;
+				}
+				else if (NumBytesToSkip > 0)
+				{
+					int64 NumSkipped = ReadData(nullptr, NumBytesToSkip, -1);
+					if (NumSkipped != NumBytesToSkip)
+					{
+						FAccessUnit::Release(AccessUnit);
+						AccessUnit = nullptr;
+						if (!HasReadBeenAborted())
+						{
+							SetError(FString::Printf(TEXT("Failed to skip over %lld bytes in segment \"%s\""), (long long int)NumBytesToSkip, *RequestURL), INTERNAL_SEG_ERROR_BAD_MEDIA_SEGMENT);
+						}
+						break;
+					}
+				}
+				if (MoofInfo.PayloadStartOffset == 0)
+				{
+					MoofInfo.PayloadStartOffset = GetCurrentOffset();
+				}
+
+				int64 NumRead = ReadData(AccessUnit->AUData, AccessUnit->AUSize, -1);
+				if (NumRead == AccessUnit->AUSize)
+				{
+					MoofInfo.NumKeyframeBytes += AccessUnit->bIsSyncSample ? AccessUnit->AUSize : 0;
+					// Only update duration for the primary stream type.
+					if (InRequest->GetType() == CurrentlyActiveTrackData->StreamType)
+					{
+						MoofInfo.ContentDuration += Duration;
+					}
+					CurrentlyActiveTrackData->DurationSuccessfullyRead += Duration;
+
+					// If we need to decrypt we have to wait for the decrypter to become ready.
+					if (bIsSampleEncrypted && Decrypter.IsValid())
+					{
+						while(!bTerminate && !HasReadBeenAborted() &&
+								(Decrypter->GetState() == ElectraCDM::ECDMState::WaitingForKey || Decrypter->GetState() == ElectraCDM::ECDMState::Idle))
+						{
+							FMediaRunnable::SleepMilliseconds(100);
+						}
+						ElectraCDM::ECDMError DecryptResult = ElectraCDM::ECDMError::Failure;
+						if (Decrypter->GetState() == ElectraCDM::ECDMState::Ready)
+						{
+							SCOPE_CYCLE_COUNTER(STAT_ElectraPlayer_Common_StreamReader);
+							CSV_SCOPED_TIMING_STAT(ElectraPlayer, Common_StreamReader);
+							DecryptResult = Decrypter->DecryptInPlace((uint8*) AccessUnit->AUData, (int32) AccessUnit->AUSize, stp.SampleEncryptionInfo);
+						}
+						if (DecryptResult != ElectraCDM::ECDMError::Success)
+						{
+							FAccessUnit::Release(AccessUnit);
+							AccessUnit = nullptr;
+							SetError(FString::Printf(TEXT("Failed to decrypt segment \"%s\" with error %d (%s)"), *RequestURL, (int32)DecryptResult, *Decrypter->GetLastErrorMessage()), INTERNAL_SEG_ERROR_FAILED_TO_DECRYPT);
+							break;
+						}
+					}
+				}
+				else
+				{
+					// Did not get the number of bytes we needed. Either because of a read error or because we got aborted.
+					FAccessUnit::Release(AccessUnit);
+					AccessUnit = nullptr;
+					bHasErrored = true;
+					break;
+				}
+
+				// Check if the AU is outside the time range we are allowed to read.
+				// The last one (the one that is already outside the range, actually) is tagged as such and sent into the buffer.
+				// The respective decoder has to handle this flag if necessary and/or drop the AU.
+				// We need to send at least one AU down so the FMultiTrackAccessUnitBuffer does not stay empty for this period!
+				if (AccessUnit)
+				{
+					// Already sent the last one?
+					if (!CurrentlyActiveTrackData->bReadPastLastPTS)
+					{
+						CurrentlyActiveTrackData->AddAccessUnit(AccessUnit);
+					}
+					FAccessUnit::Release(AccessUnit);
+					AccessUnit = nullptr;
+				}
+
+				// Shall we pass on any AUs we already read?
+				if (bAllowEarlyEmitting && bTimeCheckPassed)
+				{
+					EmitSamples(EEmitType::UntilBlocked);
+				}
+
+				CurrentlyActiveTrackData->bIsFirstInSequence = false;
+				++StreamTypeAUCount[stIdx];
+
+				bool bHaveNextSample = stp.Iterator->Next();
+				if (!bHaveNextSample)
+				{
+					NextTrackIndex = -1;
+					stp.Iterator.Reset();
+				}
+				else if (--SampleChunkInfo.NumSamples == 0)
+				{
+					NextTrackIndex = -1;
+				}
+			}
+		}
+
+		// If we have not yet reached the end of the `mdat` because of tracks we are not interested in
+		// we need to skip over the remaining data in order to reach the following `moof` if there is one.
+		int64 NumExcessMdatBytes = LastByteOffsetOfMDAT - GetCurrentOffset();
+		check(NumExcessMdatBytes >= 0);
+		if (NumExcessMdatBytes > 0 && !bDone && !HasErrored() && !HasReadBeenAborted())
+		{
+			int64 NumSkipped = ReadData(nullptr, NumExcessMdatBytes, -1);
+			if (NumSkipped != NumExcessMdatBytes)
+			{
+				if (!HasReadBeenAborted())
+				{
+					SetError(FString::Printf(TEXT("Failed to skip over %lld bytes in segment \"%s\" to reach next moof"), (long long int)NumExcessMdatBytes, *RequestURL), INTERNAL_SEG_ERROR_BAD_MEDIA_SEGMENT);
+				}
+				break;
+			}
+		}
+		MoofInfo.PayloadEndOffset = GetCurrentOffset();
+		if (InRequest->Segment.bLowLatencyChunkedEncodingExpected)
+		{
+			DownloadStats.MovieChunkInfos.Emplace(MoveTemp(MoofInfo));
+		}
+	}
+	BoxLoader.Reset();
+	return HasReadBeenAborted() ? EHandleResult::Aborted :
+		   HasErrored() ? EHandleResult::Failed :
+		   bSkippedBecauseOfTimestampCheck ? EHandleResult::Skipped :
+		   EHandleResult::Finished;
+}
+
+FStreamSegmentReaderCommon::FStreamHandler::EHandleResult FStreamSegmentReaderCommon::FStreamHandler::HandleMP4MediaLegacy(TSharedPtrTS<FStreamSegmentRequestCommon> InRequest)
+{
+	// We need to have an mp4 init segment.
 	if (!InitSegmentData.IsType<TSharedPtrTS<const IParserISO14496_12>>())
@@ -1829,7 +2645,6 @@
 
 					const int32 stIdx = StreamTypeToArrayIndex(CurrentlyActiveTrackData->StreamType);
 
-
 					const uint32 TrackTimescale = TrackIterator->GetTimescale();
 
 					// Perform some checks and adjustments for the first access unit.
@@ -1974,8 +2789,8 @@
 							InRequest->FirstTimestampReceivedDelegate.ExecuteIfBound(InRequest);
 						}
 
-						HandleMP4Metadata(InRequest, MP4Parser, InitSegment, TimeOffset);
-						HandleMP4EventMessages(InRequest, MP4Parser);
+						HandleMP4MetadataLegacy(InRequest, MP4Parser, InitSegment, TimeOffset);
+						HandleMP4EventMessagesLegacy(InRequest, MP4Parser);
 						// If the track uses encryption we update the DRM system with the PSSH boxes that are currently in use.
 						// NOTE: This works only if there is a single decrypter for any of the multiplexed tracks.
 						//       If there is ever more than one this needs to be extended.
@@ -3511,7 +4326,6 @@
 		StreamSelector->ReportDownloadEnd(DownloadStats);
 	}
 
-
 	if (SideloadedData.IsValid())
 	{
 		TSortedMap<uint64, TSharedPtrTS<FActiveTrackData>> NewTrackDataMap;
@@ -3862,7 +4676,7 @@
 	}
 }
 
-void FStreamSegmentReaderCommon::FStreamHandler::HandleMP4EventMessages(const TSharedPtrTS<FStreamSegmentRequestCommon>& InRequest, const TSharedPtrTS<IParserISO14496_12>& InMP4Parser)
+void FStreamSegmentReaderCommon::FStreamHandler::HandleMP4EventMessagesLegacy(const TSharedPtrTS<FStreamSegmentRequestCommon>& InRequest, const TSharedPtrTS<IParserISO14496_12>& InMP4Parser)
 {
 	if (InRequest->StreamingProtocol != FStreamSegmentRequestCommon::EStreamingProtocol::DASH)
 	{
@@ -3953,7 +4767,7 @@
 	}
 }
 
-void FStreamSegmentReaderCommon::FStreamHandler::HandleMP4Metadata(const TSharedPtrTS<FStreamSegmentRequestCommon>& InRequest, const TSharedPtrTS<IParserISO14496_12>& InMP4Parser, const TSharedPtrTS<const IParserISO14496_12>& InMP4InitSegment, const FTimeValue& InBaseTime)
+void FStreamSegmentReaderCommon::FStreamHandler::HandleMP4MetadataLegacy(const TSharedPtrTS<FStreamSegmentRequestCommon>& InRequest, const TSharedPtrTS<IParserISO14496_12>& InMP4Parser, const TSharedPtrTS<const IParserISO14496_12>& InMP4InitSegment, const FTimeValue& InBaseTime)
 {
 	// Get the metadata from the movie fragment or the init segment.
 	const IParserISO14496_12::IMetadata* MoofMetadata = InMP4Parser->GetMetadata(IParserISO14496_12::EBaseBoxType::Moof);
@@ -3978,6 +4792,122 @@
 	}
 }
 
+void FStreamSegmentReaderCommon::FStreamHandler::HandleMP4Metadata(const TSharedPtrTS<FStreamSegmentRequestCommon>& InRequest, const TSharedPtr<MP4Boxes::FMP4BoxMETA, ESPMode::ThreadSafe>& InMoovMetaBox, const TSharedPtr<MP4Boxes::FMP4BoxMETA, ESPMode::ThreadSafe>& InMoofMetaBox)
+{
+	if (!InMoovMetaBox && !InMoofMetaBox)
+	{
+		return;
+	}
+	const TSharedPtr<MP4Boxes::FMP4BoxMETA, ESPMode::ThreadSafe>& MetaBox = InMoovMetaBox ? InMoovMetaBox : InMoofMetaBox;
+	if (auto HdlrBox = MetaBox->GetHandler())
+	{
+		uint32 hdlr = HdlrBox->GetHandlerType();
+		uint32 res0 = HdlrBox->GetReservedValue(0);
 
+		TArray<MP4Utilities::FMP4BoxInfo> Boxes;
+		for(auto &ch : MetaBox->GetChildren())
+		{
+			Boxes.Emplace(ch->GetBoxInfo());
+		}
+		TSharedPtr<MP4Utilities::FMP4MetadataParser, ESPMode::ThreadSafe> MediaMetadata = MakeShared<MP4Utilities::FMP4MetadataParser, ESPMode::ThreadSafe>();
+		if (MediaMetadata->Parse(hdlr, res0, Boxes) == FMetadataParser::EResult::Success)
+		{
+			FTimeValue StartTime = InRequest->GetFirstPTS();
+			StartTime.SetSequenceIndex(InRequest->TimestampSequenceIndex);
+  			PlayerSessionService->SendMessageToPlayer(FPlaylistMetadataUpdateMessage::Create(StartTime, MediaMetadata, false));
+		}
+	}
+}
+
+void FStreamSegmentReaderCommon::FStreamHandler::HandleMP4EventMessages(const TSharedPtrTS<FStreamSegmentRequestCommon>& InRequest, const TSharedPtr<MP4Boxes::FMP4BoxSIDX, ESPMode::ThreadSafe>& InSidxBox, const TArray<TSharedPtr<MP4Boxes::FMP4BoxEMSG, ESPMode::ThreadSafe>>& InEventMessageBoxes)
+{
+	if (InRequest->StreamingProtocol != FStreamSegmentRequestCommon::EStreamingProtocol::DASH || InEventMessageBoxes.IsEmpty())
+	{
+		return;
+	}
+
+	SCOPE_CYCLE_COUNTER(STAT_ElectraPlayer_Common_StreamReader);
+	CSV_SCOPED_TIMING_STAT(ElectraPlayer, Common_StreamReader);
+	// Are there `emsg` boxes we need to handle?
+	FTimeValue AbsPeriodStart = InRequest->PeriodStart + InRequest->AST + InRequest->AdditionalAdjustmentTime;
+	// We may need the EPT from the `sidx` if there is one.
+	for(int32 nEmsg=0, nEmsgs=InEventMessageBoxes.Num(); nEmsg<nEmsgs; ++nEmsg)
+	{
+		const TSharedPtr<MP4Boxes::FMP4BoxEMSG, ESPMode::ThreadSafe>& Emsg = InEventMessageBoxes[nEmsg];
+		// This event must be described by an <InbandEventStream> in order to be processed.
+		for(int32 i=0; i<InRequest->Segment.InbandEventStreams.Num(); ++i)
+		{
+			const FSegmentInformationCommon::FInbandEventStream& ibs = InRequest->Segment.InbandEventStreams[i];
+			if (ibs.SchemeIdUri.Equals(Emsg->GetSchemeIdUri()) &&
+				(ibs.Value.IsEmpty() || ibs.Value.Equals(Emsg->GetValue())))
+			{
+				TSharedPtrTS<DASH::FPlayerEvent> NewEvent = MakeSharedTS<DASH::FPlayerEvent>();
+				NewEvent->SetOrigin(IAdaptiveStreamingPlayerAEMSEvent::EOrigin::InbandEventStream);
+				NewEvent->SetSchemeIdUri(Emsg->GetSchemeIdUri());
+				NewEvent->SetValue(Emsg->GetValue());
+				NewEvent->SetID(LexToString(Emsg->GetID()));
+				uint32 Timescale = Emsg->GetTimescale();
+				uint32 Duration = Emsg->GetEventDuration();
+				FTimeValue PTS;
+				if (Emsg->GetVersion() == 0)
+				{
+					// Version 0 uses a presentation time delta relative to the EPT of the SIDX, if it exists, or if not
+					// to the PTS of the first AU, which should be the same as the segment media start time.
+					FTimeValue PTD((int64)Emsg->GetPresentationTimeDelta(), Timescale);
+					FTimeValue EPT;
+					if (InSidxBox)
+					{
+						EPT.SetFromND((int64)InSidxBox->GetEarliestPresentationTime(), InSidxBox->GetTimescale());
+					}
+					else
+					{
+						EPT.SetFromND((int64)InRequest->Segment.Time, InRequest->Segment.Timescale);
+					}
+					FTimeValue PTO(InRequest->Segment.PTO, InRequest->Segment.Timescale);
+					PTS = AbsPeriodStart - PTO + EPT + PTD;
+				}
+				else if (Emsg->GetVersion() == 1)
+				{
+					FTimeValue EventTime(Emsg->GetPresentationTime(), Timescale);
+					FTimeValue PTO(FTimeValue::GetZero());
+					PTS = AbsPeriodStart - PTO + EventTime;
+				}
+				NewEvent->SetPresentationTime(PTS);
+				if (Duration != 0xffffffff)
+				{
+					NewEvent->SetDuration(FTimeValue((int64)Duration, Timescale));
+				}
+				NewEvent->SetMessageData(Emsg->GetMessageData());
+				NewEvent->SetPeriodID(InRequest->Period->GetUniqueIdentifier());
+				// Add the event to the handler.
+				if (PTS.IsValid())
+				{
+					/*
+						Check that we have no seen this event in this segment already. This is for the case where the 'emsg' appears
+						inbetween multiple 'moof' boxes. As per ISO/IEC 23009-1:2019 Section 5.10.3.3.1 General:
+							A Media Segment if based on the ISO BMFF container may contain one or more event message ('emsg') boxes. If present, any 'emsg' box shall be placed as follows:
+							- It may be placed before the first 'moof' box of the segment.
+							- It may be placed in between any 'mdat' and 'moof' box. In this case, an equivalent 'emsg' with the same id value shall be present before the first 'moof' box of any Segment.
+					*/
+					int32 EventIdx = SegmentEventsFound.IndexOfByPredicate([&NewEvent](const TSharedPtrTS<DASH::FPlayerEvent>& This) {
+						return NewEvent->GetSchemeIdUri().Equals(This->GetSchemeIdUri()) &&
+								NewEvent->GetID().Equals(This->GetID()) &&
+								(NewEvent->GetValue().IsEmpty() || NewEvent->GetValue().Equals(This->GetValue()));
+					});
+					if (EventIdx == INDEX_NONE)
+					{
+						PlayerSessionService->GetAEMSEventHandler()->AddEvent(NewEvent, InRequest->Period->GetUniqueIdentifier(), IAdaptiveStreamingPlayerAEMSHandler::EEventAddMode::AddIfNotExists);
+						SegmentEventsFound.Emplace(MoveTemp(NewEvent));
+					}
+				}
+				break;
+			}
+		}
+	}
+}
+
+
+
+
 FStreamSegmentReaderCommon::FStreamHandler::EEmitResult FStreamSegmentReaderCommon::FStreamHandler::EmitSamples(EEmitType InEmitType)
 {
@@ -4372,4 +5302,3 @@
 
 
 } // namespace Electra
-
diff --git a/Engine/Plugins/Media/ElectraPlayer/Source/ElectraPlayerRuntime/Private/Runtime/Player/StreamSegmentReaderCommon.h b/Engine/Plugins/Media/ElectraPlayer/Source/ElectraPlayerRuntime/Private/Runtime/Player/StreamSegmentReaderCommon.h
--- a/Engine/Plugins/Media/ElectraPlayer/Source/ElectraPlayerRuntime/Private/Runtime/Player/StreamSegmentReaderCommon.h
+++ b/Engine/Plugins/Media/ElectraPlayer/Source/ElectraPlayerRuntime/Private/Runtime/Player/StreamSegmentReaderCommon.h
@@ -5,18 +5,142 @@
 #include "Player/PlayerStreamReader.h"
 #include "Player/AdaptiveStreamingPlayerResourceRequest.h"
 #include "Player/DASH/PlayerEventDASH.h"
+#include "MP4Boxes.h"
+#include "MP4Utilities.h"
+#include "Utilities/MP4Helpers.h"
 #include "Demuxer/ParserISO14496-12.h"
 #include "Demuxer/ParserMKV.h"
 #include "HTTP/HTTPManager.h"
 #include "Player/DRM/DRMManager.h"
 #include "Misc/TVariant.h"
+#include "HAL/IConsoleManager.h"
 
 
 
 namespace Electra
 {
+namespace Stream
+{
+extern TAutoConsoleVariable<bool> CVarElectraStreamUseLegacyMP4;
+}
+
 struct FMediaPlaylistAndStateHLS;
 
+class FParsedMOOVCommon
+{
+public:
+	virtual ~FParsedMOOVCommon() = default;
+
+	struct FTrackInfo
+	{
+		TSharedPtr<MP4Boxes::FMP4BoxTRAK, ESPMode::ThreadSafe> TrakBox;
+		FCodecTypeFormat CodecFormat;
+		FDRMTypeFormat DrmFormat;
+		FDecoderInformation DecoderInformation;
+		uint32 TrackID = 0;
+		uint32 TrackTimescale = 0;
+	};
+
+	bool Parse(TConstArrayView<uint8> InInitSegmentData)
+	{
+		TSharedPtr<MP4Boxes::FMP4BoxMOOV, ESPMode::ThreadSafe> MOOV;
+		TArray<TSharedPtr<MP4Utilities::FMP4BoxData, ESPMode::ThreadSafe>> RootBoxes;
+		if (MP4Utilities::FMP4BoxLocator::LocateAndReadRootBoxesFromBuffer(RootBoxes, InInitSegmentData))
+		{
+			TArray<TSharedPtr<MP4Boxes::FMP4BoxBase, ESPMode::ThreadSafe>> ParsedRootBoxes;
+			ParsedRootBoxes.SetNum(RootBoxes.Num());
+			for(int32 i=0; i<RootBoxes.Num(); ++i)
+			{
+				MP4Boxes::FMP4BoxTreeParser tp;
+				if (tp.ParseBoxTree(RootBoxes[i]))
+				{
+					ParsedRootBoxes[i] = tp.GetBoxTree();
+					if (ParsedRootBoxes[i]->GetType() == MP4Utilities::MakeBoxAtom('m','o','o','v'))
+					{
+						MOOV = StaticCastSharedPtr<MP4Boxes::FMP4BoxMOOV>(ParsedRootBoxes[i]);
+						break;
+					}
+				}
+				else
+				{
+					return false;
+				}
+			}
+			if (MOOV)
+			{
+				return Parse(MOOV);
+			}
+		}
+		return false;
+	}
+
+	bool Parse(TSharedPtr<MP4Boxes::FMP4BoxMOOV, ESPMode::ThreadSafe> InMOOV)
+	{
+		if (!InMOOV)
+		{
+			return false;
+		}
+		MoovBox = InMOOV;
+		MvexBox = MoovBox->FindBoxRecursive<MP4Boxes::FMP4BoxMVEX>(MP4Utilities::MakeBoxAtom('m','v','e','x'), 0);
+		if (!MvexBox)
+		{
+			return false;
+		}
+		auto Udta = MoovBox->FindBoxRecursive<MP4Boxes::FMP4BoxUDTA>(MP4Utilities::MakeBoxAtom('u','d','t','a'), 0);
+		MetaBox = Udta ? Udta->FindBoxRecursive<MP4Boxes::FMP4BoxMETA>(MP4Utilities::MakeBoxAtom('m','e','t','a'), 0) : nullptr;
+		TArray<TSharedPtr<MP4Boxes::FMP4BoxTRAK, ESPMode::ThreadSafe>> AllTracks;
+		MoovBox->GetAllBoxInstances(AllTracks, MP4Utilities::MakeBoxAtom('t','r','a','k'));
+		for(int32 nTk=0; nTk<AllTracks.Num(); ++nTk)
+		{
+			auto Tkhd = AllTracks[nTk]->FindBoxRecursive<MP4Boxes::FMP4BoxTKHD>(MP4Utilities::MakeBoxAtom('t','k','h','d'), 0);
+			auto Mdhd = AllTracks[nTk]->FindBoxRecursive<MP4Boxes::FMP4BoxMDHD>(MP4Utilities::MakeBoxAtom('m','d','h','d'), 1);
+			if (!Tkhd || !Mdhd)
+			{
+				continue;
+			}
+			FTrackInfo ti;
+			ti.TrakBox = AllTracks[nTk];
+			ti.TrackID = Tkhd->GetTrackID();
+			ti.TrackTimescale = Mdhd->GetTimescale();
+
+			FString Msg;
+			if (!GetTrackFormatInfo(Msg, ti.CodecFormat, ti.DrmFormat, ti.DecoderInformation, AllTracks[nTk], ti.TrackID, true))
+			{
+				continue;
+			}
+			if (!ti.DecoderInformation.bIsDecodable)
+			{
+				continue;
+			}
+			if (ti.CodecFormat.Type != FCodecTypeFormat::EType::Video && ti.CodecFormat.Type != FCodecTypeFormat::EType::Audio && ti.CodecFormat.Type != FCodecTypeFormat::EType::Subtitle)
+			{
+				continue;
+			}
+			TrackMap.Add(ti.TrackID, ti);
+		}
+		MoovBox->GetAllBoxInstances(PsshBoxes, MP4Utilities::MakeBoxAtom('p','s','s','h'));
+		return true;
+	}
+
+	const TMap<uint32, FTrackInfo>& GetTrackMap() const
+	{ return TrackMap; }
+
+	TSharedPtr<MP4Boxes::FMP4BoxMVEX, ESPMode::ThreadSafe> GetMVEXBox() const
+	{ return MvexBox; }
+
+	TSharedPtr<MP4Boxes::FMP4BoxMETA, ESPMode::ThreadSafe> GetMETABox() const
+	{ return MetaBox; }
+	const TArray<TSharedPtr<MP4Boxes::FMP4BoxPSSH, ESPMode::ThreadSafe>>& GetPSSHBoxes() const
+	{ return PsshBoxes; }
+private:
+	TSharedPtr<MP4Boxes::FMP4BoxMOOV, ESPMode::ThreadSafe> MoovBox;
+	TSharedPtr<MP4Boxes::FMP4BoxMVEX, ESPMode::ThreadSafe> MvexBox;
+	TSharedPtr<MP4Boxes::FMP4BoxMETA, ESPMode::ThreadSafe> MetaBox;
+	TArray<TSharedPtr<MP4Boxes::FMP4BoxPSSH, ESPMode::ThreadSafe>> PsshBoxes;
+	TMap<uint32, FTrackInfo> TrackMap;
+};
+
+
 
 struct FProducerReferenceTimeInfo : public IProducerReferenceTimeInfo
 {
@@ -544,12 +668,37 @@ private:
 			int64 PrevDTS90k = 0;
 		};
 
+		class FBoxDataLoader : public MP4Utilities::IMP4DataReaderBase, public TSharedFromThis<FBoxDataLoader, ESPMode::ThreadSafe>
+		{
+		public:
+			FBoxDataLoader(IGenericDataReader* InGenericReader)
+				: GenericReader(InGenericReader)
+			{ }
+			virtual ~FBoxDataLoader() = default;
+			int64 ReadData(void* InOutBuffer, int64 InNumBytes, int64 InFromOffset, FCancellationCheckDelegate InCheckCancellationDelegate) override
+			{
+				check(InFromOffset == GetCurrentFileOffset());
+				int64 nr = GenericReader->ReadData(InOutBuffer, InNumBytes, InFromOffset);
+				return nr == InNumBytes ? nr : GenericReader->HasReadBeenAborted() ? EResult::Canceled : EResult::ReadError;
+			}
+			int64 GetTotalFileSize() override
+			{ return GenericReader->GetTotalSize(); }
+			int64 GetCurrentFileOffset() override
+			{ return GenericReader->GetCurrentOffset(); }
+			bool HasReachedEOF() override
+			{ return GenericReader->HasReachedEOF(); }
+			FString GetLastError() override
+			{ return FString(); }
+		private:
+			IGenericDataReader* GenericReader = nullptr;
+		};
+
 		FStreamHandler();
 		virtual ~FStreamHandler();
 		void Cancel(bool bSilent);
 		void SignalWork();
 		void WorkerThread();
-		void SetError(const FString& InMessage, uint16 InCode);
+		EHandleResult SetError(const FString& InMessage, uint16 InCode);
 		bool HasErrored() const;
 		void LogMessage(IInfoLog::ELevel Level, const FString& Message);
 		int32 HTTPProgressCallback(const IElectraHttpManager::FRequest* Request);
@@ -563,6 +712,7 @@ private:
 		void HandleCommonMediaBegin(TSharedPtrTS<FStreamSegmentRequestCommon> InRequest);
 		void HandleCommonMediaEnd(EHandleResult InSegmentResult, TSharedPtrTS<FStreamSegmentRequestCommon> InRequest);
 		EHandleResult HandleMP4Media(TSharedPtrTS<FStreamSegmentRequestCommon> InRequest);
+		EHandleResult HandleMP4MediaLegacy(TSharedPtrTS<FStreamSegmentRequestCommon> InRequest);
 		EHandleResult HandleMKVMedia(TSharedPtrTS<FStreamSegmentRequestCommon> InRequest);
 		EHandleResult HandleTSMedia(TSharedPtrTS<FStreamSegmentRequestCommon> InRequest);
 		EHandleResult HandleID3RawMedia(TSharedPtrTS<FStreamSegmentRequestCommon> InRequest, int32 InID3HeaderSize);
@@ -571,8 +721,10 @@ private:
 		EHandleResult HandleFillerDataSetup(TSharedPtrTS<FStreamSegmentRequestCommon> InRequest);
 		void InsertFillerData(TSharedPtrTS<FActiveTrackData> InActiveTrackData, TSharedPtrTS<FStreamSegmentRequestCommon> InRequest);
 		void CheckForInbandDASHEvents(const TSharedPtrTS<FStreamSegmentRequestCommon>& InRequest);
-		void HandleMP4Metadata(const TSharedPtrTS<FStreamSegmentRequestCommon>& InRequest, const TSharedPtrTS<IParserISO14496_12>& InMP4Parser, const TSharedPtrTS<const IParserISO14496_12>& InMP4InitSegment, const FTimeValue& InBaseTime);
-		void HandleMP4EventMessages(const TSharedPtrTS<FStreamSegmentRequestCommon>& InRequest, const TSharedPtrTS<IParserISO14496_12>& InMP4Parser);
+		void HandleMP4Metadata(const TSharedPtrTS<FStreamSegmentRequestCommon>& InRequest, const TSharedPtr<MP4Boxes::FMP4BoxMETA, ESPMode::ThreadSafe>& InMoovMetaBox, const TSharedPtr<MP4Boxes::FMP4BoxMETA, ESPMode::ThreadSafe>& InMoofMetaBox);
+		void HandleMP4EventMessages(const TSharedPtrTS<FStreamSegmentRequestCommon>& InRequest, const TSharedPtr<MP4Boxes::FMP4BoxSIDX, ESPMode::ThreadSafe>& InSidxBox, const TArray<TSharedPtr<MP4Boxes::FMP4BoxEMSG, ESPMode::ThreadSafe>>& InEventMessageBoxes);
+		void HandleMP4MetadataLegacy(const TSharedPtrTS<FStreamSegmentRequestCommon>& InRequest, const TSharedPtrTS<IParserISO14496_12>& InMP4Parser, const TSharedPtrTS<const IParserISO14496_12>& InMP4InitSegment, const FTimeValue& InBaseTime);
+		void HandleMP4EventMessagesLegacy(const TSharedPtrTS<FStreamSegmentRequestCommon>& InRequest, const TSharedPtrTS<IParserISO14496_12>& InMP4Parser);
 
 		void SelectPrimaryTrackData(const TSharedPtrTS<FStreamSegmentRequestCommon>& InRequest);
 		EEmitResult EmitSamples(EEmitType InEmitType);
@@ -614,7 +766,7 @@ private:
 		FString ABRAbortReason;
 
 		struct FNoInitSegmentData {int32 Zero=0;};
-		TVariant<FNoInitSegmentData, TSharedPtrTS<const IParserISO14496_12>, TSharedPtrTS<const IParserMKV>, TSharedPtrTS<const TArray<uint8>>> InitSegmentData;
+		TVariant<FNoInitSegmentData, TSharedPtrTS<FParsedMOOVCommon>, TSharedPtrTS<const IParserISO14496_12>, TSharedPtrTS<const IParserMKV>, TSharedPtrTS<const TArray<uint8>>> InitSegmentData;
 		Metrics::FSegmentDownloadStats DownloadStats;
 		HTTP::FConnectionInfo CurrentConnectionInfo;
 		TSharedPtrTS<IElectraHttpManager::FProgressListener> ProgressListener;
diff --git a/Engine/Plugins/Media/ElectraPlayer/Source/ElectraPlayerRuntime/Private/Runtime/Player/isobmff/PlaylistISOBMFF.cpp b/Engine/Plugins/Media/ElectraPlayer/Source/ElectraPlayerRuntime/Private/Runtime/Player/isobmff/PlaylistISOBMFF.cpp
--- a/Engine/Plugins/Media/ElectraPlayer/Source/ElectraPlayerRuntime/Private/Runtime/Player/isobmff/PlaylistISOBMFF.cpp
+++ b/Engine/Plugins/Media/ElectraPlayer/Source/ElectraPlayerRuntime/Private/Runtime/Player/isobmff/PlaylistISOBMFF.cpp
@@ -738,7 +738,7 @@ IManifest::FResult FManifestISOBMFFInternal::FPlayPeriodISOBMFF::GetStartingSegm
 	{
 		IManifest::FResult res(IManifest::FResult::EType::NotFound);
 		res.SetErrorDetail(FErrorDetail().SetError(UEMEDIA_ERROR_DETAIL)
-							.SetFacility(Facility::EFacility::MP4Playlist)
+							.SetFacility(Facility::EFacility::ISOBMFFPlaylist)
 							.SetCode(ERRCODE_MANIFEST_ISOBMFF_STARTSEGMENT_NOT_FOUND)
 							.SetMessage(FString::Printf(TEXT("Could not find start segment for time %lld"), (long long int)StartPosition.Time.GetAsHNS())));
 		return res;
@@ -840,6 +840,7 @@ IManifest::FResult FManifestISOBMFFInternal::FPlayPeriodISOBMFF::GetStartingSegm
 			Requests[i]->EarliestPTS = StartRequest->EarliestPTS;
 			Requests[i]->LastPTS = StartRequest->LastPTS;
 			Requests[i]->TimestampSequenceIndex = StartRequest->TimestampSequenceIndex;
+			Requests[i]->DrmClient = Requests[i]->Track.TrackInfo->DrmFormat.IsEncrypted() ? ma->GetDrmClient() : nullptr;
 
 			if (LoaderType == FTimelineAssetISOBMFF::ELoaderType::Chunks)
 			{
@@ -927,6 +928,7 @@ IManifest::FResult FManifestISOBMFFInternal::FPlayPeriodISOBMFF::GetStartingSegm
 					TrkReq->EarliestPTS = StartRequest->EarliestPTS;
 					TrkReq->LastPTS = StartRequest->LastPTS;
 					TrkReq->TimestampSequenceIndex = StartRequest->TimestampSequenceIndex;
+					TrkReq->DrmClient = TrkReq->Track.TrackInfo->DrmFormat.IsEncrypted() ? ma->GetDrmClient() : nullptr;
 					TrkReq->bWasSelectedAtStart = AS->GetUniqueIdentifier().Equals(SelectedStreamTypeIDs[stIdx]);
 					// Add this track request to the list of dependent streams.
 					// Note: This track could be at EOS at that time already, which is indicated by it not having an iterator!
@@ -1149,7 +1151,7 @@ void FManifestISOBMFFInternal::FPlayPeriodISOBMFF::IncreaseSegmentFetchDelay(con
 /**
  * Builds the timeline asset.
  */
-FErrorDetail FManifestISOBMFFInternal::FTimelineAssetISOBMFF::Build(IPlayerSessionServices* InPlayerSessionServices, TArray<TSharedPtr<MP4Boxes::FMP4BoxBase>>&& InParsedRootBoxes, const FString& InURL,
+FErrorDetail FManifestISOBMFFInternal::FTimelineAssetISOBMFF::Build(IPlayerSessionServices* InPlayerSessionServices, TArray<TSharedPtr<MP4Boxes::FMP4BoxBase, ESPMode::ThreadSafe>>&& InParsedRootBoxes, const FString& InURL,
 																	TSharedPtr<FMP4DataLoader, ESPMode::ThreadSafe> InDataLoader, FMP4DataLoader::FCancellationCheckDelegate InCancelCheck)
 {
 	PlayerSessionServices = InPlayerSessionServices;
@@ -1157,15 +1159,18 @@ FErrorDetail FManifestISOBMFFInternal::FTimelineAssetISOBMFF::Build(IPlayerSessi
 	MediaURL = InURL;
 	bIsLocalFile = MediaURL.StartsWith(TEXT("file:"), ESearchCase::IgnoreCase);
 
+	TArray<ElectraCDM::IMediaCDM::FCDMCandidate> DRMCandidates;
+	auto DRMManager = PlayerSessionServices->GetDRMManager();
+
 	FErrorDetail Error;
 	auto SetError = [&Error](const FString& InMsg) -> void
 	{
-		Error.SetFacility(Facility::EFacility::MP4Playlist);
+		Error.SetFacility(Facility::EFacility::ISOBMFFPlaylist);
 		Error.SetMessage(InMsg);
 		Error.SetCode(ERRCODE_MANIFEST_ISOBMFF_NO_PLAYABLE_STREAMS);
 	};
 
-	auto FindRootBox = [&](uint32 InType) -> TSharedPtr<MP4Boxes::FMP4BoxBase>
+	auto FindRootBox = [&](uint32 InType) -> TSharedPtr<MP4Boxes::FMP4BoxBase, ESPMode::ThreadSafe>
 	{
 		auto rb = ParsedRootBoxes.FindByPredicate([Type=InType](const TSharedPtr<MP4Boxes::FMP4BoxBase>& In) { return In->GetType() == Type;});
 		return rb ? *rb : nullptr;
@@ -1281,10 +1286,47 @@ FErrorDetail FManifestISOBMFFInternal::FTimelineAssetISOBMFF::Build(IPlayerSessi
 		{
 			continue;
 		}
-		// Skip encrypted tracks.
-		if (TrkInf->DrmFormat.bIsEncrypted)
+		// Can we decode this encrypted track?
+		if (TrkInf->DrmFormat.IsEncrypted())
 		{
-			continue;
+			if (!TrkInf->DrmFormat.EncryptionInfo.IsType<FDRMTypeFormat::FISOEncryptionInfo>() || !DRMManager)
+			{
+				continue;
+			}
+			const FDRMTypeFormat::FISOEncryptionInfo& EncInfo(TrkInf->DrmFormat.EncryptionInfo.Get<FDRMTypeFormat::FISOEncryptionInfo>());
+			if (!EncInfo.Scheme)
+			{
+				continue;
+			}
+			FString Scheme = FDRMManager::MakePrintableStringFromUint32(EncInfo.Scheme);
+			TSharedPtr<ElectraCDM::IMediaCDMCapabilities, ESPMode::ThreadSafe> DRMCapabilities;
+			bool bCanBeDecrypted = false;
+			for(int32 nI=0; nI<EncInfo.CDMInfos.Num(); ++nI)
+			{
+				FString uuid = FDRMManager::MakeHexStringFromArray(EncInfo.CDMInfos[nI].SystemID);
+				DRMCapabilities = DRMManager->GetCDMCapabilitiesForScheme(uuid, FString(), FString());
+				if (DRMCapabilities.IsValid())
+				{
+					ElectraCDM::IMediaCDMCapabilities::ESupportResult Result = DRMCapabilities->SupportsCipher(Scheme);
+					if (Result == ElectraCDM::IMediaCDMCapabilities::ESupportResult::Supported)
+					{
+						bCanBeDecrypted = true;
+						ElectraCDM::IMediaCDM::FCDMCandidate cand;
+						cand.SchemeId = uuid;
+						cand.CommonScheme = Scheme;
+						for(int32 j=0; j<EncInfo.CDMInfos[nI].KIDs.Num(); ++j)
+						{
+							cand.DefaultKIDs.Emplace(FDRMManager::MakeHexStringFromArray(EncInfo.CDMInfos[nI].KIDs[j]));
+						}
+						DRMCandidates.Emplace(MoveTemp(cand));
+					}
+				}
+			}
+			// If this track cannot be decrypted we skip over it.
+			if (!bCanBeDecrypted)
+			{
+				continue;
+			}
 		}
 
 		// Can this track be decoded and meets the current playback restrictions?
@@ -1621,6 +1663,24 @@ FErrorDetail FManifestISOBMFFInternal::FTimelineAssetISOBMFF::Build(IPlayerSessi
 			}
 		}
 	}
+
+	// Request licenses
+	if (DRMCandidates.Num())
+	{
+		ElectraCDM::ECDMError Result = DRMManager->CreateDRMClient(DrmClient, DRMCandidates);
+		if (Result == ElectraCDM::ECDMError::Success && DrmClient.IsValid())
+		{
+			DrmClient->RegisterEventListener(DRMManager);
+			DrmClient->SetLicenseServerURL(MediaURL);
+			DrmClient->PrepareLicenses();
+		}
+		else
+		{
+			Error.SetFacility(Facility::EFacility::ISOBMFFPlaylist);
+			Error.SetMessage(FString::Printf(TEXT("Failed to create DRM client with error %d"), (int32)Result));
+			Error.SetCode(ERRCODE_MANIFEST_ISOBMFF_DRM_ERROR);
+		}
+	}
 	return Error;
 }
 
@@ -1629,7 +1689,7 @@ void FManifestISOBMFFInternal::FTimelineAssetISOBMFF::LogMessage(IInfoLog::ELeve
 {
 	if (PlayerSessionServices)
 	{
-		PlayerSessionServices->PostLog(Facility::EFacility::MP4Playlist, Level, Message);
+		PlayerSessionServices->PostLog(Facility::EFacility::ISOBMFFPlaylist, Level, Message);
 	}
 }
 
diff --git a/Engine/Plugins/Media/ElectraPlayer/Source/ElectraPlayerRuntime/Private/Runtime/Player/isobmff/PlaylistISOBMFF.h b/Engine/Plugins/Media/ElectraPlayer/Source/ElectraPlayerRuntime/Private/Runtime/Player/isobmff/PlaylistISOBMFF.h
--- a/Engine/Plugins/Media/ElectraPlayer/Source/ElectraPlayerRuntime/Private/Runtime/Player/isobmff/PlaylistISOBMFF.h
+++ b/Engine/Plugins/Media/ElectraPlayer/Source/ElectraPlayerRuntime/Private/Runtime/Player/isobmff/PlaylistISOBMFF.h
@@ -189,7 +189,7 @@ public:
 
 		virtual ~FTimelineAssetISOBMFF() = default;
 
-		FErrorDetail Build(IPlayerSessionServices* InPlayerSessionServices, TArray<TSharedPtr<MP4Boxes::FMP4BoxBase>>&& InParsedRootBoxes, const FString& InURL,
+		FErrorDetail Build(IPlayerSessionServices* InPlayerSessionServices, TArray<TSharedPtr<MP4Boxes::FMP4BoxBase, ESPMode::ThreadSafe>>&& InParsedRootBoxes, const FString& InURL,
 						   TSharedPtr<FMP4DataLoader, ESPMode::ThreadSafe> InDataLoader, FMP4DataLoader::FCancellationCheckDelegate InCancelCheck);
 
 		FTimeRange GetTimeRange() const override
@@ -332,6 +332,9 @@ public:
 		ELoaderType GetLoaderType() const
 		{ return LoaderType; }
 
+		TSharedPtrTS<ElectraCDM::IMediaCDMClient> GetDrmClient() const
+		{ return DrmClient; }
+
 	private:
 		struct FPlayRangeEndInfo
 		{
@@ -359,7 +362,7 @@ public:
 		FPlayRangeEndInfo PlayRangeEndInfo;
 		// Metadata from the 'meta' box, if any.
 		TSharedPtrTS<FMetadataParser> MediaMetadata;
-		//TSharedPtrTS<ElectraCDM::IMediaCDMClient> DrmClient;
+		TSharedPtrTS<ElectraCDM::IMediaCDMClient> DrmClient;
 	};
 
 
diff --git a/Engine/Plugins/Media/ElectraPlayer/Source/ElectraPlayerRuntime/Private/Runtime/Player/isobmff/PlaylistReaderISOBMFF.cpp b/Engine/Plugins/Media/ElectraPlayer/Source/ElectraPlayerRuntime/Private/Runtime/Player/isobmff/PlaylistReaderISOBMFF.cpp
--- a/Engine/Plugins/Media/ElectraPlayer/Source/ElectraPlayerRuntime/Private/Runtime/Player/isobmff/PlaylistReaderISOBMFF.cpp
+++ b/Engine/Plugins/Media/ElectraPlayer/Source/ElectraPlayerRuntime/Private/Runtime/Player/isobmff/PlaylistReaderISOBMFF.cpp
@@ -195,7 +195,7 @@ void FPlaylistReaderISOBMFF::WorkerThread()
 	const TArray<uint32> FirstBoxes { MP4Utilities::MakeBoxAtom('f','t','y','p'), MP4Utilities::MakeBoxAtom('s','t','y','p'), MP4Utilities::MakeBoxAtom('s','i','d','x'), MP4Utilities::MakeBoxAtom('f','r','e','e'), MP4Utilities::MakeBoxAtom('s','k','i','p') };
 	const TArray<uint32> StopAfterBoxes { MP4Utilities::MakeBoxAtom('m','o','o','v') };	// Empty means to read all boxes and not stop after a specific one.
 	const TArray<uint32> ReadBoxes;			// Empty means to read in all boxes except `mdat`
-	bool bGotBoxes = BoxLocator.LocateAndReadRootBoxes(RootBoxes, DataLoader, FirstBoxes, StopAfterBoxes, ReadBoxes, DataLoaderCancel);
+	bool bGotBoxes = BoxLocator.LocateAndReadRootBoxes(RootBoxes, DataLoader, FirstBoxes, StopAfterBoxes, false, ReadBoxes, DataLoaderCancel);
 	ConnectionInfo = DataLoader->GetConnectionInfo();
 	bool bHasErrored = ConnectionInfo.StatusInfo.ErrorDetail.IsSet();
 
@@ -230,7 +230,7 @@ void FPlaylistReaderISOBMFF::WorkerThread()
 			}
 
 			// Parse all the root boxes. There are typically no more than 4 to 6.
-			TArray<TSharedPtr<MP4Boxes::FMP4BoxBase>> ParsedRootBoxes;
+			TArray<TSharedPtr<MP4Boxes::FMP4BoxBase, ESPMode::ThreadSafe>> ParsedRootBoxes;
 			ParsedRootBoxes.SetNum(RootBoxes.Num());
 			bool bIsFragmented = false;
 			for(int32 i=0; i<RootBoxes.Num(); ++i)
@@ -269,7 +269,7 @@ void FPlaylistReaderISOBMFF::WorkerThread()
 				const TArray<uint32> FrgFirstBoxes;
 				const TArray<uint32> FrgStopAfterBoxes;	// Empty means to read all boxes and not stop after a specific one.
 				const TArray<uint32> FrgReadBoxes;		// Empty means to read in all boxes except `mdat`
-				bGotBoxes = BoxLocator.LocateAndReadRootBoxes(FrgRootBoxes, DataLoader, FrgFirstBoxes, FrgStopAfterBoxes, FrgReadBoxes, DataLoaderCancel);
+				bGotBoxes = BoxLocator.LocateAndReadRootBoxes(FrgRootBoxes, DataLoader, FrgFirstBoxes, FrgStopAfterBoxes, false, FrgReadBoxes, DataLoaderCancel);
 				ConnectionInfo = DataLoader->GetConnectionInfo();
 				bHasErrored = ConnectionInfo.StatusInfo.ErrorDetail.IsSet();
 				if (BoxLocator.GetLastError().Len())
diff --git a/Engine/Plugins/Media/ElectraPlayer/Source/ElectraPlayerRuntime/Private/Runtime/Player/isobmff/StreamReaderISOBMFF.cpp b/Engine/Plugins/Media/ElectraPlayer/Source/ElectraPlayerRuntime/Private/Runtime/Player/isobmff/StreamReaderISOBMFF.cpp
--- a/Engine/Plugins/Media/ElectraPlayer/Source/ElectraPlayerRuntime/Private/Runtime/Player/isobmff/StreamReaderISOBMFF.cpp
+++ b/Engine/Plugins/Media/ElectraPlayer/Source/ElectraPlayerRuntime/Private/Runtime/Player/isobmff/StreamReaderISOBMFF.cpp
@@ -355,6 +355,9 @@ void FStreamReaderISOBMFF::FStreamHandler::HandleMuxedRequest()
 	{
 		TSharedPtr<FAccessUnit::CodecData, ESPMode::ThreadSafe> CSD;
 		TSharedPtr<FBufferSourceInfo, ESPMode::ThreadSafe> BufferInfo;
+		TSharedPtr<ElectraCDM::IMediaCDMDecrypter, ESPMode::ThreadSafe> Decrypter;
+		FDRMTypeFormat DrmFormat;
+		ElectraCDM::FMediaCDMSampleInfo SampleEncryptionInfo;
 		TArray<FAccessUnit*> AccessUnits;
 		FTimeValue DurationSuccessfullyRead { FTimeValue::GetZero() };
 		FTimeValue DurationSuccessfullyDelivered { FTimeValue::GetZero() };
@@ -363,6 +366,7 @@ void FStreamReaderISOBMFF::FStreamHandler::HandleMuxedRequest()
 		bool bEndOfTrack { false };
 		bool bDeliveredLastAccessUnit { false };
 		bool bNotifiedOfTrackEnd { false };
+		bool bIsEncrypted { false };
 		int32 DependendStreamIdx { -1 };
 	};
 	TMap<uint32, FTrackVars> TrackMap;
@@ -394,6 +398,29 @@ void FStreamReaderISOBMFF::FStreamHandler::HandleMuxedRequest()
 			tv.bNotifiedOfTrackEnd = true;
 			Parameters.EventListener->OnFragmentReachedEOS(tv.StreamType, tv.BufferInfo);
 		}
+		// Encrypted?
+		if (Request->DependentStreams[i]->Track.TrackInfo->DrmFormat.IsEncrypted())
+		{
+			tv.bIsEncrypted = true;
+			tv.DrmFormat = Request->DependentStreams[i]->Track.TrackInfo->DrmFormat;
+			if (tv.DrmFormat.EncryptionInfo.IsType<FDRMTypeFormat::FISOEncryptionInfo>())
+			{
+				const FDRMTypeFormat::FISOEncryptionInfo& encInf(tv.DrmFormat.EncryptionInfo.Get<FDRMTypeFormat::FISOEncryptionInfo>());
+				tv.SampleEncryptionInfo.DefaultKID = encInf.DefaultKID;
+				tv.SampleEncryptionInfo.IV = encInf.DefaultIV;
+				tv.SampleEncryptionInfo.Scheme4CC = encInf.Scheme;
+				if (encInf.BlockPattern.IsSet())
+				{
+					tv.SampleEncryptionInfo.Pattern.PatternType = 1;
+					tv.SampleEncryptionInfo.Pattern.CryptByteBlock = encInf.BlockPattern.GetValue().CryptByteBlock;
+					tv.SampleEncryptionInfo.Pattern.SkipByteBlock = encInf.BlockPattern.GetValue().SkipByteBlock;
+				}
+				if (Request->DependentStreams[i]->DrmClient)
+				{
+					Request->DependentStreams[i]->DrmClient->CreateDecrypter(tv.Decrypter, FString());
+				}
+			}
+		}
 		TrackMap.Emplace(Request->DependentStreams[i]->Track.TrackInfo->TrackID, MoveTemp(tv));
 	}
 
@@ -625,6 +652,50 @@ void FStreamReaderISOBMFF::FStreamHandler::HandleMuxedRequest()
 				FScopeLock lock(ReceiveBuffer->GetLock());
 				FMemory::Memcpy(AccessUnit->AUData, reinterpret_cast<const uint8*>(ReceiveBuffer->GetLinearReadData()) + SampleRelativeOffset, SampleSize);
 			}
+			// Decrypt?
+			if (TrackVars.bIsEncrypted)
+			{
+				if (TrackVars.Decrypter)
+				{
+					// Wait for the decrypter to get ready.
+					while(!ShouldStop() && (TrackVars.Decrypter->GetState() == ElectraCDM::ECDMState::WaitingForKey || TrackVars.Decrypter->GetState() == ElectraCDM::ECDMState::Idle))
+					{
+						FMediaRunnable::SleepMilliseconds(100);
+					}
+					// Set up the sample encryption info.
+					const MP4Boxes::FMP4BoxSENC::FEntry& enc = TrackIt->GetEncryptionInfo();
+					TrackVars.SampleEncryptionInfo.SubSamples.Empty();
+					TrackVars.SampleEncryptionInfo.IV = enc.IV;
+					for(auto &sub : enc.SubSamples)
+					{
+						TrackVars.SampleEncryptionInfo.SubSamples.Add({.NumClearBytes = sub.NumClearBytes, .NumEncryptedBytes = sub.NumEncryptedBytes});
+					}
+					// Decrypt
+					ElectraCDM::ECDMError DecryptResult = ElectraCDM::ECDMError::Failure;
+					if (TrackVars.Decrypter->GetState() == ElectraCDM::ECDMState::Ready)
+					{
+						DecryptResult = TrackVars.Decrypter->DecryptInPlace((uint8*) AccessUnit->AUData, (int32) AccessUnit->AUSize, TrackVars.SampleEncryptionInfo);
+					}
+					if (DecryptResult != ElectraCDM::ECDMError::Success)
+					{
+						FAccessUnit::Release(AccessUnit);
+						AccessUnit = nullptr;
+						ds.FailureReason = FString::Printf(TEXT("Failed to decrypt segment with error %d (%s)"), (int32)DecryptResult, *TrackVars.Decrypter->GetLastErrorMessage());
+						bHasErrored = true;
+						bDone = true;
+						break;
+					}
+				}
+				else
+				{
+					FAccessUnit::Release(AccessUnit);
+					AccessUnit = nullptr;
+					ds.FailureReason = FString::Printf(TEXT("No valid decrypter to decrypt track"));
+					bHasErrored = true;
+					bDone = true;
+					break;
+				}
+			}
 			if (AccessUnit->DTS >= AccessUnit->LatestPTS && AccessUnit->PTS >= AccessUnit->LatestPTS)
 			{
 				AccessUnit->bIsLastInPeriod = true;
@@ -680,7 +751,10 @@ void FStreamReaderISOBMFF::FStreamHandler::HandleMuxedRequest()
 	ds.TimeToDownload = (Request->ConnectionInfo.RequestEndTime - Request->ConnectionInfo.RequestStartTime).GetAsSeconds();
 	ds.ByteSize = Request->ConnectionInfo.ContentLength;
 	ds.NumBytesDownloaded = Request->ConnectionInfo.BytesReadSoFar;
-	ds.FailureReason = Request->ConnectionInfo.StatusInfo.ErrorDetail.GetMessage();
+	if (ds.FailureReason.IsEmpty())
+	{
+		ds.FailureReason = Request->ConnectionInfo.StatusInfo.ErrorDetail.GetMessage();
+	}
 	ds.bWasSuccessful = !bHasErrored;
 	ds.URL.URL = Request->ConnectionInfo.EffectiveURL;
 	if (Request->ReportingIndexInMux >= 0)
@@ -770,6 +844,42 @@ void FStreamReaderISOBMFF::FStreamHandler::HandleChunkOrFragmentRequest()
 	BufferSource->Codec = CSD->ParsedInfo.GetCodecName();
 	BufferSource->Kind = Request->KindOfTrack;
 
+	struct FDecryptVars
+	{
+		ElectraCDM::FMediaCDMSampleInfo SampleEncryptionInfo;
+		FDRMTypeFormat DrmFormat;
+		TSharedPtr<ElectraCDM::IMediaCDMDecrypter, ESPMode::ThreadSafe> Decrypter;
+		bool bIsEncrypted { false };
+	};
+	FDecryptVars DecryptVars;
+	if (Request->Track.TrackInfo->DrmFormat.IsEncrypted())
+	{
+		DecryptVars.bIsEncrypted = true;
+		DecryptVars.DrmFormat = Request->Track.TrackInfo->DrmFormat;
+		if (DecryptVars.DrmFormat.EncryptionInfo.IsType<FDRMTypeFormat::FISOEncryptionInfo>())
+		{
+			const FDRMTypeFormat::FISOEncryptionInfo& encInf(DecryptVars.DrmFormat.EncryptionInfo.Get<FDRMTypeFormat::FISOEncryptionInfo>());
+			DecryptVars.SampleEncryptionInfo.DefaultKID = encInf.DefaultKID;
+			DecryptVars.SampleEncryptionInfo.IV = encInf.DefaultIV;
+			DecryptVars.SampleEncryptionInfo.Scheme4CC = encInf.Scheme;
+			if (encInf.BlockPattern.IsSet())
+			{
+				DecryptVars.SampleEncryptionInfo.Pattern.PatternType = 1;
+				DecryptVars.SampleEncryptionInfo.Pattern.CryptByteBlock = encInf.BlockPattern.GetValue().CryptByteBlock;
+				DecryptVars.SampleEncryptionInfo.Pattern.SkipByteBlock = encInf.BlockPattern.GetValue().SkipByteBlock;
+			}
+			if (Request->DrmClient)
+			{
+				Request->DrmClient->CreateDecrypter(DecryptVars.Decrypter, FString());
+			}
+		}
+	}
+
+	auto ShouldStop = [&]() -> bool
+	{
+		return HasErrored() || HasBeenAborted() || bTerminate;
+	};
+
 	uint32 PlaybackSequenceID = Request->GetPlaybackSequenceID();
 	auto TrackIt = Request->Track.TrackIterator;
 	int64 SizeNeededInBuffer = 0;
@@ -777,13 +887,13 @@ void FStreamReaderISOBMFF::FStreamHandler::HandleChunkOrFragmentRequest()
 	FTimeValue DurationSuccessfullyDelivered(FTimeValue::GetZero());
 	bool bIsFirstInSequence = true;
 	bool bSentLastAccessUnit = false;
-	while(!bDone && !HasErrored() && !HasBeenAborted() && !bTerminate)
+	while(!bDone && !ShouldStop())
 	{
 		// Get the number of bytes we need for the next sample.
 		int64 SampleSize = TrackIt->GetSampleSize();
 		SizeNeededInBuffer = TrackIt->GetSampleFileOffset() + SampleSize - StartOffset;
 		// Wait until the required amount of data has been read.
-		while(!HasErrored() && !HasBeenAborted() && !bTerminate)
+		while(!ShouldStop())
 		{
 			// Check periodically if we are to abort in case we do not have enough data yet.
 			if (ReceiveBuffer->WaitUntilSizeAvailable(SizeNeededInBuffer, 1000 * 100))
@@ -792,7 +902,7 @@ void FStreamReaderISOBMFF::FStreamHandler::HandleChunkOrFragmentRequest()
 			}
 		}
 		// Leave the loop if we have to abort.
-		if (!(!HasErrored() && !HasBeenAborted() && !bTerminate))
+		if (ShouldStop())
 		{
 			break;
 		}
@@ -841,6 +951,52 @@ void FStreamReaderISOBMFF::FStreamHandler::HandleChunkOrFragmentRequest()
 				FMemory::Memcpy(AccessUnit->AUData, reinterpret_cast<const uint8*>(ReceiveBuffer->GetLinearReadData()) + SampleBufferOffset, SampleSize);
 			}
 
+			// Decrypt?
+			if (DecryptVars.bIsEncrypted)
+			{
+				if (DecryptVars.Decrypter)
+				{
+					// Wait for the decrypter to get ready.
+					while(!ShouldStop() && (DecryptVars.Decrypter->GetState() == ElectraCDM::ECDMState::WaitingForKey || DecryptVars.Decrypter->GetState() == ElectraCDM::ECDMState::Idle))
+					{
+						FMediaRunnable::SleepMilliseconds(100);
+					}
+					// Set up the sample encryption info.
+					const MP4Boxes::FMP4BoxSENC::FEntry& enc = TrackIt->GetEncryptionInfo();
+					DecryptVars.SampleEncryptionInfo.SubSamples.Empty();
+					DecryptVars.SampleEncryptionInfo.IV = enc.IV;
+					for(auto &sub : enc.SubSamples)
+					{
+						DecryptVars.SampleEncryptionInfo.SubSamples.Add({.NumClearBytes = sub.NumClearBytes, .NumEncryptedBytes = sub.NumEncryptedBytes});
+					}
+					// Decrypt
+					ElectraCDM::ECDMError DecryptResult = ElectraCDM::ECDMError::Failure;
+					if (DecryptVars.Decrypter->GetState() == ElectraCDM::ECDMState::Ready)
+					{
+						DecryptResult = DecryptVars.Decrypter->DecryptInPlace((uint8*) AccessUnit->AUData, (int32) AccessUnit->AUSize, DecryptVars.SampleEncryptionInfo);
+					}
+					if (DecryptResult != ElectraCDM::ECDMError::Success)
+					{
+						FAccessUnit::Release(AccessUnit);
+						AccessUnit = nullptr;
+						ds.FailureReason = FString::Printf(TEXT("Failed to decrypt segment with error %d (%s)"), (int32)DecryptResult, *DecryptVars.Decrypter->GetLastErrorMessage());
+						bHasErrored = true;
+						bDone = true;
+						break;
+					}
+				}
+				else
+				{
+					FAccessUnit::Release(AccessUnit);
+					AccessUnit = nullptr;
+					ds.FailureReason = FString::Printf(TEXT("No valid decrypter to decrypt track"));
+					bHasErrored = true;
+					bDone = true;
+					break;
+				}
+			}
+
+
 			if (AccessUnit->DTS >= AccessUnit->LatestPTS && AccessUnit->PTS >= AccessUnit->LatestPTS)
 			{
 				// Tag the last one and send it off, but stop doing so for the remainder of the segment.
@@ -895,7 +1051,10 @@ void FStreamReaderISOBMFF::FStreamHandler::HandleChunkOrFragmentRequest()
 	ds.TimeToDownload = (Request->ConnectionInfo.RequestEndTime - Request->ConnectionInfo.RequestStartTime).GetAsSeconds();
 	ds.ByteSize = Request->ConnectionInfo.ContentLength;
 	ds.NumBytesDownloaded = Request->ConnectionInfo.BytesReadSoFar;
-	ds.FailureReason = Request->ConnectionInfo.StatusInfo.ErrorDetail.GetMessage();
+	if (ds.FailureReason.IsEmpty())
+	{
+		ds.FailureReason = Request->ConnectionInfo.StatusInfo.ErrorDetail.GetMessage();
+	}
 	ds.bWasSuccessful = !bHasErrored;
 	ds.URL.URL = Request->ConnectionInfo.EffectiveURL;
 	ds.DurationDownloaded = DurationSuccessfullyRead.GetAsSeconds();
diff --git a/Engine/Plugins/Media/ElectraPlayer/Source/ElectraPlayerRuntime/Private/Runtime/Player/isobmff/StreamReaderISOBMFF.h b/Engine/Plugins/Media/ElectraPlayer/Source/ElectraPlayerRuntime/Private/Runtime/Player/isobmff/StreamReaderISOBMFF.h
--- a/Engine/Plugins/Media/ElectraPlayer/Source/ElectraPlayerRuntime/Private/Runtime/Player/isobmff/StreamReaderISOBMFF.h
+++ b/Engine/Plugins/Media/ElectraPlayer/Source/ElectraPlayerRuntime/Private/Runtime/Player/isobmff/StreamReaderISOBMFF.h
@@ -101,7 +101,6 @@ public:
 
 	// Encryption
 	TSharedPtrTS<ElectraCDM::IMediaCDMClient> DrmClient;
-	FString DrmMimeType;
 
 	// Download stats that are set during and after download.
 	Metrics::FSegmentDownloadStats DownloadStats;
@@ -157,7 +156,6 @@ private:
 		volatile bool bHasErrored = false;
 
 		TSharedPtrTS<FWaitableBuffer> ReceiveBuffer;
-		TSharedPtr<ElectraCDM::IMediaCDMDecrypter, ESPMode::ThreadSafe> Decrypter;
 
 		static std::atomic<uint32> UniqueDownloadID;
 	};
diff --git a/Engine/Plugins/Media/ElectraPlayer/Source/ElectraPlayerRuntime/Private/Runtime/Utilities/MP4Helpers.cpp b/Engine/Plugins/Media/ElectraPlayer/Source/ElectraPlayerRuntime/Private/Runtime/Utilities/MP4Helpers.cpp
--- a/Engine/Plugins/Media/ElectraPlayer/Source/ElectraPlayerRuntime/Private/Runtime/Utilities/MP4Helpers.cpp
+++ b/Engine/Plugins/Media/ElectraPlayer/Source/ElectraPlayerRuntime/Private/Runtime/Utilities/MP4Helpers.cpp
@@ -333,6 +333,58 @@ bool GetTrackFormatInfo(FString& OutMessage, FCodecTypeFormat& OutCodecInfo, FDR
 		}
 	};
 
+	auto SetupEncryption = [&](uint32& InOutSampleType, const TSharedPtr<MP4Boxes::FMP4BoxBase, ESPMode::ThreadSafe>& InSample) -> bool
+	{
+		// Locate the encryption related boxes
+		auto Frma = InSample->FindBoxRecursive<MP4Boxes::FMP4BoxFRMA>(MP4Utilities::MakeBoxAtom('f','r','m','a'), 2);
+		if (!Frma)
+		{
+			OutMessage = FString::Printf(TEXT("Encrypted track #%u is missing mandatory `frma` box, ignoring this track."), InTrackID);
+			return false;
+		}
+		InOutSampleType = Frma->GetOriginalFormat();
+
+		OutDRMInfo.EncryptionInfo.Emplace<FDRMTypeFormat::FISOEncryptionInfo>();
+		FDRMTypeFormat::FISOEncryptionInfo& iso = OutDRMInfo.EncryptionInfo.Get<FDRMTypeFormat::FISOEncryptionInfo>();
+		// Optional `schm` box.
+		if (auto Schm = InSample->FindBoxRecursive<MP4Boxes::FMP4BoxSCHM>(MP4Utilities::MakeBoxAtom('s','c','h','m'), 2))
+		{
+			iso.Scheme = Schm->GetSchemeType();
+			iso.SchemeVersion = Schm->GetSchemeVersion();
+		}
+		// Optional `tenc` box.
+		if (auto Tenc = InSample->FindBoxRecursive<MP4Boxes::FMP4BoxTENC>(MP4Utilities::MakeBoxAtom('t','e','n','c'), 3))
+		{
+			if (Tenc->HasDefaultCryptBlockValues())
+			{
+				FDRMTypeFormat::FISOEncryptionInfo::FBlockPattern ptn;
+				ptn.CryptByteBlock = Tenc->GetDefaultCryptByteBlock();
+				ptn.SkipByteBlock = Tenc->GetDefaultSkipByteBlock();
+				iso.BlockPattern = MoveTemp(ptn);
+			}
+			iso.DefaultKID = Tenc->GetDefaultKID();
+			iso.DefaultIV = Tenc->GetDefaultConstantIV();
+			iso.DefaultIVSize = Tenc->GetDefaultPerSampleIVSize();
+		}
+		// Optional `pssh` boxes.
+		if (auto Moov = InSample->FindParentBox<MP4Boxes::FMP4BoxMOOV>(MP4Utilities::MakeBoxAtom('m','o','o','v')))
+		{
+			TArray<TSharedPtr<MP4Boxes::FMP4BoxPSSH, ESPMode::ThreadSafe>> PsshBoxes;
+			Moov->GetAllBoxInstances<MP4Boxes::FMP4BoxPSSH>(PsshBoxes, MP4Utilities::MakeBoxAtom('p','s','s','h'));
+			for(auto &pssh : PsshBoxes)
+			{
+				FDRMTypeFormat::FISOEncryptionInfo::FCDMInfo& cdm = iso.CDMInfos.Emplace_GetRef();
+				cdm.SystemID = pssh->GetSystemID();
+				cdm.Data = pssh->GetData();
+				for(auto &kid : pssh->GetKIDs())
+				{
+					cdm.KIDs.Emplace(kid);
+				}
+			}
+		}
+		return true;
+	};
+
 	// Get the sample type of this track. This call is necessary to trigger parsing of child nodes.
 	MP4Boxes::FMP4BoxSampleEntry::ESampleType SampleType = Stsd->GetSampleType();
 	// If already knows to not be supported, skip it.
@@ -354,7 +406,6 @@ bool GetTrackFormatInfo(FString& OutMessage, FCodecTypeFormat& OutCodecInfo, FDR
 	}
 
 	OutDecoderInfo.bIsDecodable = true;
-	OutDRMInfo.bIsEncrypted = false;
 	// Based on the sample type, get it and see if it is using codec we support.
 	if (SampleType == MP4Boxes::FMP4BoxSampleEntry::ESampleType::Video)
 	{
@@ -373,47 +424,44 @@ bool GetTrackFormatInfo(FString& OutMessage, FCodecTypeFormat& OutCodecInfo, FDR
 			return false;
 		}
 
-		switch(Visuals[0]->GetType())
+		uint32 VisualType = Visuals[0]->GetType();
+		// Encrypted?
+		if (VisualType == MP4Utilities::MakeBoxAtom('e','n','c','v'))
 		{
-			case MP4Utilities::MakeBoxAtom('e','n','c','v'):
+			if (!SetupEncryption(VisualType, Visuals[0]))
 			{
-				OutDRMInfo.bIsEncrypted = true;
-				break;
+				return false;
 			}
-			default:
+		}
+		if (Visuals[0]->GetFrameCount() == 1)
+		{
+			OutCodecInfo.Type = Electra::FCodecTypeFormat::EType::Video;
+			OutCodecInfo.FourCC = VisualType;
+			OutCodecInfo.Properties.Emplace<Electra::FCodecTypeFormat::FVideo>();
+			Electra::FCodecTypeFormat::FVideo& cvid = OutCodecInfo.Properties.Get<Electra::FCodecTypeFormat::FVideo>();
+			cvid.Width = Visuals[0]->GetWidth();
+			cvid.Height = Visuals[0]->GetHeight();
+			auto Stts = InTrack->FindBoxRecursive<MP4Boxes::FMP4BoxSTTS>(MP4Utilities::MakeBoxAtom('s','t','t','s'), 5);
+			if (Mdhd.IsValid() && Stts.IsValid() && Stts->GetEntries().Num())
 			{
-				if (Visuals[0]->GetFrameCount() == 1)
-				{
-					OutCodecInfo.Type = Electra::FCodecTypeFormat::EType::Video;
-					OutCodecInfo.FourCC = Visuals[0]->GetType();
-					OutCodecInfo.Properties.Emplace<Electra::FCodecTypeFormat::FVideo>();
-					Electra::FCodecTypeFormat::FVideo& cvid = OutCodecInfo.Properties.Get<Electra::FCodecTypeFormat::FVideo>();
-					cvid.Width = Visuals[0]->GetWidth();
-					cvid.Height = Visuals[0]->GetHeight();
-					auto Stts = InTrack->FindBoxRecursive<MP4Boxes::FMP4BoxSTTS>(MP4Utilities::MakeBoxAtom('s','t','t','s'), 5);
-					if (Mdhd.IsValid() && Stts.IsValid() && Stts->GetEntries().Num())
-					{
-						cvid.FrameRate = FFrameRate(Mdhd->GetTimescale(), Stts->GetEntries()[0].sample_delta);
-					}
-					for(auto& Ch : Visuals[0]->GetChildren())
-					{
-						OutCodecInfo.ExtraBoxes.Emplace(Ch->GetType(), Ch->GetBoxData());
-					}
-					if (!CodecFormatUtils->PrepareCodecTypeFormat(OutCodecInfo, IElectraCodecFormatUtilsModularFeature::FCodecTypeFormatParams()))
-					{
-						OutMessage = FString::Printf(TEXT("Failed to prepare the codec format type, ignoring track #%u."), InTrackID);
-						return false;
-					}
-					// Check if this video is decodable.
-					CheckIfDecodable();
-				}
-				else
-				{
-					OutMessage = FString::Printf(TEXT("Track #%u has a frame count other than 1 in the VisualSampleEntry, ignoring this track."), InTrackID);
-					return false;
-				}
-				break;
+				cvid.FrameRate = FFrameRate(Mdhd->GetTimescale(), Stts->GetEntries()[0].sample_delta);
+			}
+			for(auto& Ch : Visuals[0]->GetChildren())
+			{
+				OutCodecInfo.ExtraBoxes.Emplace(Ch->GetType(), Ch->GetBoxData());
+			}
+			if (!CodecFormatUtils->PrepareCodecTypeFormat(OutCodecInfo, IElectraCodecFormatUtilsModularFeature::FCodecTypeFormatParams()))
+			{
+				OutMessage = FString::Printf(TEXT("Failed to prepare the codec format type, ignoring track #%u."), InTrackID);
+				return false;
 			}
+			// Check if this video is decodable.
+			CheckIfDecodable();
+		}
+		else
+		{
+			OutMessage = FString::Printf(TEXT("Track #%u has a frame count other than 1 in the VisualSampleEntry, ignoring this track."), InTrackID);
+			return false;
 		}
 	}
 	else if (SampleType == MP4Boxes::FMP4BoxSampleEntry::ESampleType::Audio)
@@ -433,36 +481,32 @@ bool GetTrackFormatInfo(FString& OutMessage, FCodecTypeFormat& OutCodecInfo, FDR
 			return false;
 		}
 
-		switch(Audios[0]->GetType())
+		uint32 AudioType = Audios[0]->GetType();
+		// Encrypted?
+		if (AudioType == MP4Utilities::MakeBoxAtom('e','n','c','a'))
 		{
-			case MP4Utilities::MakeBoxAtom('e','n','c','a'):
+			if (!SetupEncryption(AudioType, Audios[0]))
 			{
-				OutDRMInfo.bIsEncrypted = true;
-				break;
-			}
-			default:
-			{
-				OutCodecInfo.Type = Electra::FCodecTypeFormat::EType::Audio;
-				OutCodecInfo.FourCC = Audios[0]->GetType();
-				OutCodecInfo.Properties.Emplace<Electra::FCodecTypeFormat::FAudio>();
-				Electra::FCodecTypeFormat::FAudio& caud = OutCodecInfo.Properties.Get<Electra::FCodecTypeFormat::FAudio>();
-				caud.NumChannels = Audios[0]->GetChannelCount();
-				caud.SampleRate = Audios[0]->GetSampleRate();
-				for(auto& Ch : Audios[0]->GetChildren())
-				{
-					OutCodecInfo.ExtraBoxes.Emplace(Ch->GetType(), Ch->GetBoxData());
-				}
-				if (!CodecFormatUtils->PrepareCodecTypeFormat(OutCodecInfo, IElectraCodecFormatUtilsModularFeature::FCodecTypeFormatParams()))
-				{
-					OutMessage = FString::Printf(TEXT("Failed to prepare the codec format type, ignoring track #%u."), InTrackID);
-					return false;
-				}
-				// Check if this video is decodable.
-				CheckIfDecodable();
-				break;
+				return false;
 			}
 		}
-
+		OutCodecInfo.Type = Electra::FCodecTypeFormat::EType::Audio;
+		OutCodecInfo.FourCC = AudioType;
+		OutCodecInfo.Properties.Emplace<Electra::FCodecTypeFormat::FAudio>();
+		Electra::FCodecTypeFormat::FAudio& caud = OutCodecInfo.Properties.Get<Electra::FCodecTypeFormat::FAudio>();
+		caud.NumChannels = Audios[0]->GetChannelCount();
+		caud.SampleRate = Audios[0]->GetSampleRate();
+		for(auto& Ch : Audios[0]->GetChildren())
+		{
+			OutCodecInfo.ExtraBoxes.Emplace(Ch->GetType(), Ch->GetBoxData());
+		}
+		if (!CodecFormatUtils->PrepareCodecTypeFormat(OutCodecInfo, IElectraCodecFormatUtilsModularFeature::FCodecTypeFormatParams()))
+		{
+			OutMessage = FString::Printf(TEXT("Failed to prepare the codec format type, ignoring track #%u."), InTrackID);
+			return false;
+		}
+		// Check if this audio is decodable.
+		CheckIfDecodable();
 	}
 	else if (SampleType == MP4Boxes::FMP4BoxSampleEntry::ESampleType::Subtitles || SampleType == MP4Boxes::FMP4BoxSampleEntry::ESampleType::TimedText)
 	{
@@ -481,34 +525,30 @@ bool GetTrackFormatInfo(FString& OutMessage, FCodecTypeFormat& OutCodecInfo, FDR
 			return false;
 		}
 
-		switch(Subtitles[0]->GetType())
+		uint32 SubtitleType = Subtitles[0]->GetType();
+		// Encrypted?
+		if (SubtitleType == MP4Utilities::MakeBoxAtom('e','n','c','t') || SubtitleType == MP4Utilities::MakeBoxAtom('e','n','c','u'))
 		{
-			case MP4Utilities::MakeBoxAtom('e','n','c','t'):
-			case MP4Utilities::MakeBoxAtom('e','n','c','u'):
+			if (!SetupEncryption(SubtitleType, Subtitles[0]))
 			{
-				OutDRMInfo.bIsEncrypted = true;
-				break;
-			}
-			default:
-			{
-				OutCodecInfo.Type = Electra::FCodecTypeFormat::EType::Subtitle;
-				OutCodecInfo.FourCC = Subtitles[0]->GetType();
-				OutCodecInfo.Properties.Emplace<Electra::FCodecTypeFormat::FSubtitle>();
-				//Electra::FCodecTypeFormat::FSubtitle& csub = OutCodecInfo.Properties.Get<Electra::FCodecTypeFormat::FSubtitle>();
-				// Add the sample box itself as an extra box.
-				OutCodecInfo.ExtraBoxes.Emplace(Subtitles[0]->GetType(), Subtitles[0]->GetBoxData());
-				for(auto& Ch : Subtitles[0]->GetChildren())
-				{
-					OutCodecInfo.ExtraBoxes.Emplace(Ch->GetType(), Ch->GetBoxData());
-				}
-				if (!CodecFormatUtils->PrepareCodecTypeFormat(OutCodecInfo, IElectraCodecFormatUtilsModularFeature::FCodecTypeFormatParams()))
-				{
-					OutMessage = FString::Printf( TEXT("Failed to prepare the codec format type, ignoring track #%u."), InTrackID);
-					return false;
-				}
-				break;
+				return false;
 			}
 		}
+		OutCodecInfo.Type = Electra::FCodecTypeFormat::EType::Subtitle;
+		OutCodecInfo.FourCC = SubtitleType;
+		OutCodecInfo.Properties.Emplace<Electra::FCodecTypeFormat::FSubtitle>();
+		//Electra::FCodecTypeFormat::FSubtitle& csub = OutCodecInfo.Properties.Get<Electra::FCodecTypeFormat::FSubtitle>();
+		// Add the sample box itself as an extra box.
+		OutCodecInfo.ExtraBoxes.Emplace(SubtitleType, Subtitles[0]->GetBoxData());
+		for(auto& Ch : Subtitles[0]->GetChildren())
+		{
+			OutCodecInfo.ExtraBoxes.Emplace(Ch->GetType(), Ch->GetBoxData());
+		}
+		if (!CodecFormatUtils->PrepareCodecTypeFormat(OutCodecInfo, IElectraCodecFormatUtilsModularFeature::FCodecTypeFormatParams()))
+		{
+			OutMessage = FString::Printf( TEXT("Failed to prepare the codec format type, ignoring track #%u."), InTrackID);
+			return false;
+		}
 	}
 	else if (SampleType == MP4Boxes::FMP4BoxSampleEntry::ESampleType::QTFFTimecode)
 	{
@@ -551,17 +591,10 @@ bool GetTrackFormatInfo(FString& OutMessage, FCodecTypeFormat& OutCodecInfo, FDR
 	// If the track has not been identified as usable so far, remove it.
 	if (OutCodecInfo.Type == Electra::FCodecTypeFormat::EType::Invalid)
 	{
-		if (!OutDRMInfo.bIsEncrypted)
-		{
-			TArray<TSharedPtr<MP4Boxes::FMP4BoxSampleEntry, ESPMode::ThreadSafe>> Entries;
-			Stsd->GetSampleDescriptions(Entries);
-			check(Entries.Num() == 1);
-			OutMessage = FString::Printf(TEXT("Track of sample type \"%s\" is not supported, ignoring track #%u."), *MP4Utilities::GetPrintableBoxAtom(Entries[0]->GetType()), InTrackID);
-		}
-		else
-		{
-			OutMessage = FString::Printf(TEXT("Track is using encryption, ignoring track #%u."), InTrackID);
-		}
+		TArray<TSharedPtr<MP4Boxes::FMP4BoxSampleEntry, ESPMode::ThreadSafe>> Entries;
+		Stsd->GetSampleDescriptions(Entries);
+		check(Entries.Num() == 1);
+		OutMessage = FString::Printf(TEXT("Track of sample type \"%s\" is not supported, ignoring track #%u."), *MP4Utilities::GetPrintableBoxAtom(Entries[0]->GetType()), InTrackID);
 		return false;
 	}
 	return true;
diff --git a/Engine/Plugins/Media/ElectraPlayer/Source/ElectraProtron/Private/Player/ElectraProtronPlayerImpl.cpp b/Engine/Plugins/Media/ElectraPlayer/Source/ElectraProtron/Private/Player/ElectraProtronPlayerImpl.cpp
--- a/Engine/Plugins/Media/ElectraPlayer/Source/ElectraProtron/Private/Player/ElectraProtronPlayerImpl.cpp
+++ b/Engine/Plugins/Media/ElectraPlayer/Source/ElectraProtron/Private/Player/ElectraProtronPlayerImpl.cpp
@@ -251,7 +251,7 @@ void FElectraProtronPlayer::FImpl::InternalOpen(const FString& InFilename)
 	const TArray<uint32> ReadBoxes;			// Empty means to read in all boxes except `mdat`
 	const TArray<uint32> StopAfterBoxes;	// Empty means to read all boxes and not stop after a specific one.
 	TArray<TSharedPtr<MP4Utilities::FMP4BoxData, ESPMode::ThreadSafe>> RootBoxes;
-	if (!BoxLocator.LocateAndReadRootBoxes(RootBoxes, Reader, FirstBoxes, StopAfterBoxes, ReadBoxes, MP4Utilities::IMP4DataReaderBase::FCancellationCheckDelegate::CreateLambda([&](){return bAbort;})))
+	if (!BoxLocator.LocateAndReadRootBoxes(RootBoxes, Reader, FirstBoxes, StopAfterBoxes, false, ReadBoxes, MP4Utilities::IMP4DataReaderBase::FCancellationCheckDelegate::CreateLambda([&](){return bAbort;})))
 	{
 		LastErrorMessage = BoxLocator.GetLastError();
 		return;
diff --git a/Engine/Plugins/Media/ElectraUtil/Source/ElectraBase/Public/CodecTypeFormat.h b/Engine/Plugins/Media/ElectraUtil/Source/ElectraBase/Public/CodecTypeFormat.h
--- a/Engine/Plugins/Media/ElectraUtil/Source/ElectraBase/Public/CodecTypeFormat.h
+++ b/Engine/Plugins/Media/ElectraUtil/Source/ElectraBase/Public/CodecTypeFormat.h
@@ -139,18 +139,33 @@ namespace Electra
 
 
 	/**
-	 *
+	 * Encryption information of a track.
 	 */
 	struct FDRMTypeFormat
 	{
-		bool bIsEncrypted = false;
+		struct FNullEncryptionInfo
+		{
+		};
 
-		struct FEncryptionInfo
+		// As defined by ISO/IEC 14496-12 Section 8.12 and ISO/IEC 23001-7
+		struct FISOEncryptionInfo
 		{
+			// Set by `schm` box, if it exists.
+			uint32 Scheme = 0;
+			uint32 SchemeVersion = 0;
+
+			// Set by `tenc` box, if it exists.
+			struct FBlockPattern
+			{
+				uint8 CryptByteBlock = 0;
+				uint8 SkipByteBlock = 0;
+			};
 			TArray<uint8> DefaultKID;
 			TArray<uint8> DefaultIV;
+			TOptional<FBlockPattern> BlockPattern;
 			int32 DefaultIVSize = 0;
-			uint32 Scheme = 0;
+
+			// Set by `pssh` boxes, if they exist.
 			struct FCDMInfo
 			{
 				TArray<uint8> SystemID;
@@ -160,6 +175,9 @@ namespace Electra
 			TArray<FCDMInfo> CDMInfos;
 		};
 
+		TVariant<FNullEncryptionInfo, FISOEncryptionInfo> EncryptionInfo;
+		bool IsEncrypted() const
+		{ return !EncryptionInfo.IsType<FNullEncryptionInfo>(); }
 	};
 
 	struct FDecoderInformation
diff --git a/Engine/Plugins/Media/MP4Utilities/Source/MP4Boxes/Private/MP4BoxBase.cpp b/Engine/Plugins/Media/MP4Utilities/Source/MP4Boxes/Private/MP4BoxBase.cpp
--- a/Engine/Plugins/Media/MP4Utilities/Source/MP4Boxes/Private/MP4BoxBase.cpp
+++ b/Engine/Plugins/Media/MP4Utilities/Source/MP4Boxes/Private/MP4BoxBase.cpp
@@ -8,13 +8,7 @@ namespace MP4Boxes
 
 	TArray<uint8> FMP4BoxBase::GetBoxDataRAW() const
 	{
-		TArray<uint8> Raw;
-		Raw.AddUninitialized(8 + BoxInfo.Data.Num());
-		uint32* bd = reinterpret_cast<uint32*>(Raw.GetData());
-		*bd++ = MP4Utilities::EndianSwap(static_cast<uint32>(Raw.Num()));
-		*bd++ = MP4Utilities::EndianSwap(BoxInfo.Type);
-		FMemory::Memcpy(bd, BoxInfo.Data.GetData(), BoxInfo.Data.Num());
-		return Raw;
+		return BoxInfo.GetBoxDataRAW();
 	}
 
 	void FMP4BoxBase::ProcessBoxChildrenRecursively(MP4Utilities::FMP4AtomReader& InOutReader, const MP4Utilities::FMP4BoxInfo& InCurrentBoxInfo)
@@ -33,7 +27,7 @@ namespace MP4Boxes
 	}
 
 
-	bool FMP4BoxTreeParser::ParseBoxTreeInternal(const TWeakPtr<FMP4BoxBase>& InParent, const MP4Utilities::FMP4BoxInfo& InBox)
+	bool FMP4BoxTreeParser::ParseBoxTreeInternal(const TWeakPtr<FMP4BoxBase, ESPMode::ThreadSafe>& InParent, const MP4Utilities::FMP4BoxInfo& InBox)
 	{
 		BoxTree = FMP4BoxFactory::Get().Create(InParent, InBox);
 		// Parse the enclosed box recursively unless it contains
diff --git a/Engine/Plugins/Media/MP4Utilities/Source/MP4Boxes/Private/MP4Boxes.cpp b/Engine/Plugins/Media/MP4Utilities/Source/MP4Boxes/Private/MP4Boxes.cpp
--- a/Engine/Plugins/Media/MP4Utilities/Source/MP4Boxes/Private/MP4Boxes.cpp
+++ b/Engine/Plugins/Media/MP4Utilities/Source/MP4Boxes/Private/MP4Boxes.cpp
@@ -7,6 +7,55 @@
 namespace MP4Boxes
 {
 
+/****************************************************************************************************************************************************/
+	struct FMP4BoxFTYP::FParsed
+	{
+		TArray<uint32> Brands;
+		uint32 MinorVersion = 0;
+	};
+	void FMP4BoxFTYP::ParseIfRequired()
+	{
+		if (Parsed.IsValid())
+		{
+			return;
+		}
+		Parsed = MakePimpl<FParsed>();
+		MP4Utilities::FMP4AtomReader ar(BoxInfo.Data);
+		uint32 Major, Minor;
+		ar.Read(Major);
+		ar.Read(Minor);
+		int32 NumCompatibleBrands = ar.GetNumBytesRemaining() / sizeof(uint32);
+		Parsed->Brands.SetNumUninitialized(1 + NumCompatibleBrands);
+		Parsed->MinorVersion = Minor;
+		Parsed->Brands[0] = Major;
+		for(int32 i=0; i<NumCompatibleBrands; ++i)
+		{
+			uint32 Compatible;
+			ar.Read(Compatible);
+			Parsed->Brands[1 + i] = Compatible;
+		}
+	}
+	uint32 FMP4BoxFTYP::GetMajorBrand()
+	{
+		ParseIfRequired();
+		return Parsed->Brands[0];
+	}
+	uint32 FMP4BoxFTYP::GetMinorVersion()
+	{
+		ParseIfRequired();
+		return Parsed->MinorVersion;
+	}
+	TConstArrayView<uint32> FMP4BoxFTYP::GetCompatibleBrands()
+	{
+		ParseIfRequired();
+		return Parsed->Brands.Num() > 1 ? MakeConstArrayView<uint32>(Parsed->Brands.GetData() + 1, Parsed->Brands.Num()-1) : TConstArrayView<uint32>();
+	}
+	TConstArrayView<uint32> FMP4BoxFTYP::GetAllBrands()
+	{
+		ParseIfRequired();
+		return Parsed->Brands;
+	}
+
 /****************************************************************************************************************************************************/
 	struct FMP4BoxMVHD::FParsed
 	{
@@ -1863,9 +1912,31 @@ namespace MP4Boxes
 		}
 	}
 
-	void FMP4BoxSAIZ::Test()
+	bool FMP4BoxSAIZ::HasAuxInfoType()
 	{
 		ParseIfRequired();
+		return !!(Parsed->Flags & 1);
+	}
+	uint32 FMP4BoxSAIZ::GetAuxInfoType()
+	{
+		ParseIfRequired();
+		return Parsed->AuxInfoType;
+	}
+	uint32 FMP4BoxSAIZ::GetAuxInfoTypeParameter()
+	{
+		ParseIfRequired();
+		return Parsed->AuxInfoTypeParameter;
+	}
+	uint32 FMP4BoxSAIZ::GetSampleCount()
+	{
+		ParseIfRequired();
+		return Parsed->SampleCount;
+	}
+	uint8 FMP4BoxSAIZ::GetSampleInfoSize(uint32 InForSampleNumber)
+	{
+		ParseIfRequired();
+		return Parsed->DefaultSampleInfoSize ? Parsed->DefaultSampleInfoSize :
+			   InForSampleNumber < Parsed->SampleCount ? Parsed->SampleInfoSizes[InForSampleNumber] : 0;
 	}
 
 /****************************************************************************************************************************************************/
@@ -1907,9 +1978,32 @@ namespace MP4Boxes
 		}
 	}
 
-	void FMP4BoxSAIO::Test()
+	bool FMP4BoxSAIO::HasAuxInfoType()
+	{
+		ParseIfRequired();
+		return !!(Parsed->Flags & 1);
+	}
+	uint32 FMP4BoxSAIO::GetAuxInfoType()
+	{
+		ParseIfRequired();
+		return Parsed->AuxInfoType;
+	}
+	uint32 FMP4BoxSAIO::GetAuxInfoTypeParameter()
+	{
+		ParseIfRequired();
+		return Parsed->AuxInfoTypeParameter;
+	}
+	uint32 FMP4BoxSAIO::GetEntryCount()
 	{
 		ParseIfRequired();
+		return Parsed->EntryCount;
+	}
+	uint64 FMP4BoxSAIO::GetEntry(uint32 InForEntryNumber)
+	{
+		ParseIfRequired();
+		return InForEntryNumber >= Parsed->EntryCount ? 0 :
+				Parsed->Offsets32 ? MP4Utilities::GetFromBigEndian(Parsed->Offsets32[InForEntryNumber]) :
+				Parsed->Offsets64 ? MP4Utilities::GetFromBigEndian(Parsed->Offsets64[InForEntryNumber]) : 0;
 	}
 
 /****************************************************************************************************************************************************/
@@ -2505,7 +2599,66 @@ namespace MP4Boxes
 	}
 
 /****************************************************************************************************************************************************/
+	struct FMP4BoxFRMA::FParsed
+	{
+		uint32 OriginalFormat = 0;
+	};
+	void FMP4BoxFRMA::ParseIfRequired()
+	{
+		if (Parsed.IsValid())
+		{
+			return;
+		}
+		Parsed = MakePimpl<FParsed>();
+		MP4Utilities::FMP4AtomReader ar(BoxInfo.Data);
+		ar.Read(Parsed->OriginalFormat);
+	}
 
+	uint32 FMP4BoxFRMA::GetOriginalFormat()
+	{
+		ParseIfRequired();
+		return Parsed->OriginalFormat;
+	}
+
+/****************************************************************************************************************************************************/
+	struct FMP4BoxSCHM::FParsed
+	{
+		FString SchemeURI;
+		uint32 SchemeType = 0;
+		uint32 SchemeVersion = 0;
+		uint32 Flags = 0;
+		uint8 Version = 0;
+	};
+	void FMP4BoxSCHM::ParseIfRequired()
+	{
+		if (Parsed.IsValid())
+		{
+			return;
+		}
+		Parsed = MakePimpl<FParsed>();
+		MP4Utilities::FMP4AtomReader ar(BoxInfo.Data);
+		ar.ReadVersionAndFlags(Parsed->Version, Parsed->Flags);
+		ar.Read(Parsed->SchemeType);
+		ar.Read(Parsed->SchemeVersion);
+		if ((Parsed->Flags & 1) != 0)
+		{
+			ar.ReadNullTerminatedStringUTF8(Parsed->SchemeURI);
+		}
+	}
+
+	uint32 FMP4BoxSCHM::GetSchemeType()
+	{
+		ParseIfRequired();
+		return Parsed->SchemeType;
+	}
+
+	uint32 FMP4BoxSCHM::GetSchemeVersion()
+	{
+		ParseIfRequired();
+		return Parsed->SchemeVersion;
+	}
+
+/****************************************************************************************************************************************************/
 	struct FMP4BoxPSSH::FParsed
 	{
 		TConstArrayView<uint8> SystemID;
@@ -2533,6 +2686,7 @@ namespace MP4Boxes
 			for(uint32 i=0; i<KID_count; ++i)
 			{
 				Parsed->KIDs[i] = MakeConstArrayView(ar.GetCurrentDataPointer(), 16);
+				ar.SkipBytes(16);
 			}
 		}
 		uint32 DataSize;
@@ -2641,11 +2795,6 @@ namespace MP4Boxes
 
 /****************************************************************************************************************************************************/
 
-	struct FMP4BoxSENC::FInitRefs
-	{
-		TArray<TSharedPtr<FMP4BoxBase, ESPMode::ThreadSafe>> RelatedBoxes;
-	};
-
 	struct FMP4BoxSENC::FParsed
 	{
 		TArray<FEntry> Entries;
@@ -2654,14 +2803,15 @@ namespace MP4Boxes
 	};
 	void FMP4BoxSENC::ParseIfRequired()
 	{
-		if (Parsed.IsValid() && InitRefs.IsValid())
+		if (Parsed.IsValid())
 		{
 			return;
 		}
 		Parsed = MakePimpl<FParsed>();
 		MP4Utilities::FMP4AtomReader ar(BoxInfo.Data);
 		ar.ReadVersionAndFlags(Parsed->Version, Parsed->Flags);
-		if (InitRefs.IsValid())
+		auto TencBox = ParentTENCBox.Pin();
+		if (TencBox)
 		{
 			// Check that we do not try to handle a pre-PIFF 1.3 box which has a different layout.
 			if ((Parsed->Flags & 1) != 0)
@@ -2685,16 +2835,6 @@ namespace MP4Boxes
 			}
 
 			// We need the `Per_Sample_IV_Size` value from the `tenc` box.
-			TSharedPtr<FMP4BoxBase, ESPMode::ThreadSafe>* TencBoxPointer = InitRefs->RelatedBoxes.FindByPredicate([](const TSharedPtr<FMP4BoxBase, ESPMode::ThreadSafe>& InBox)
-			{
-				return InBox->GetType() == MP4Utilities::MakeBoxAtom('t','e','n','c');
-			});
-			if (!TencBoxPointer)
-			{
-				UE_LOG(LogMP4Boxes, Error, TEXT("Parsing a `senc` box requires the `tenc` box of the related track to be first passed via Prepare()"));
-				return;
-			}
-			TSharedPtr<FMP4BoxTENC, ESPMode::ThreadSafe> TencBox = StaticCastSharedPtr<FMP4BoxTENC>(*TencBoxPointer);
 			const uint8 Per_Sample_IV_Size = TencBox->GetDefaultPerSampleIVSize();
 
 			const bool senc_use_subsamples = !!(Parsed->Flags & 2);
@@ -2744,18 +2884,14 @@ namespace MP4Boxes
 		}
 		else
 		{
-			UE_LOG(LogMP4Boxes, Error, TEXT("Accessing a `senc` box that has not been prepared with the corresponding `tenc` and other related boxes"));
+			UE_LOG(LogMP4Boxes, Error, TEXT("Cannot access a `senc` that has not been prepared with the corresponding `tenc`box"));
 		}
 	}
 
-	void FMP4BoxSENC::Prepare(const TArray<TSharedPtr<FMP4BoxBase, ESPMode::ThreadSafe>>& InRelatedBoxes)
+	void FMP4BoxSENC::Prepare(TWeakPtr<FMP4BoxTENC, ESPMode::ThreadSafe> InTENCBox)
 	{
-		// If already prepared we're done.
-		if (!InitRefs.IsValid())
-		{
-			InitRefs = MakePimpl<FInitRefs>();
-			InitRefs->RelatedBoxes = InRelatedBoxes;
-		}
+		ParentTENCBox = InTENCBox;
+		ParseIfRequired();
 	}
 
 	const TArray<FMP4BoxSENC::FEntry>& FMP4BoxSENC::GetEntries()
@@ -2898,6 +3034,154 @@ namespace MP4Boxes
 		return FindBoxRecursive<FMP4BoxHDLR>(MP4Utilities::MakeBoxAtom('h','d','l','r'), 0);
 	}
 
+/****************************************************************************************************************************************************/
+	struct FMP4BoxPRFT::FParsed
+	{
+		uint64 NTPTimestamp = 0;
+		uint64 MediaTime = 0;
+		uint32 ReferenceTrackID = 0;
+		uint32 Flags = 0;
+		uint8 Version = 0;
+	};
+	void FMP4BoxPRFT::ParseIfRequired()
+	{
+		if (Parsed.IsValid())
+		{
+			return;
+		}
+		Parsed = MakePimpl<FParsed>();
+		MP4Utilities::FMP4AtomReader ar(BoxInfo.Data);
+		ar.ReadVersionAndFlags(Parsed->Version, Parsed->Flags);
+		ar.Read(Parsed->ReferenceTrackID);
+		ar.Read(Parsed->NTPTimestamp);
+		if (Parsed->Version == 0)
+		{
+			uint32 v;
+			ar.Read(v);
+			Parsed->MediaTime = v;
+		}
+		else
+		{
+			ar.Read(Parsed->MediaTime);
+		}
+	}
+
+	uint32 FMP4BoxPRFT::GetFlags()
+	{
+		ParseIfRequired();
+		return Parsed->Flags;
+	}
+	uint32 FMP4BoxPRFT::GetReferenceTrackID()
+	{
+		ParseIfRequired();
+		return Parsed->ReferenceTrackID;
+	}
+	uint64 FMP4BoxPRFT::GetNTPTimestamp()
+	{
+		ParseIfRequired();
+		return Parsed->NTPTimestamp;
+	}
+	uint64 FMP4BoxPRFT::GetMediaTime()
+	{
+		ParseIfRequired();
+		return Parsed->MediaTime;
+	}
+
+/****************************************************************************************************************************************************/
+	struct FMP4BoxEMSG::FParsed
+	{
+		TArray<uint8> MessageData;
+		FString SchemeIdUri;
+		FString Value;
+		uint64 PresentationTimeOrDelta = 0;
+		uint32 Timescale = 0;
+		uint32 EventDuration = 0;
+		uint32 ID = 0;
+		uint32 Flags = 0;
+		uint8 Version = 0;
+	};
+	void FMP4BoxEMSG::ParseIfRequired()
+	{
+		if (Parsed.IsValid())
+		{
+			return;
+		}
+		Parsed = MakePimpl<FParsed>();
+		MP4Utilities::FMP4AtomReader ar(BoxInfo.Data);
+		ar.ReadVersionAndFlags(Parsed->Version, Parsed->Flags);
+		if (Parsed->Version == 0)
+		{
+			ar.ReadNullTerminatedStringUTF8(Parsed->SchemeIdUri);
+			ar.ReadNullTerminatedStringUTF8(Parsed->Value);
+			ar.Read(Parsed->Timescale);
+			uint32 PresentationTimeDelta;
+			ar.Read(PresentationTimeDelta);
+			Parsed->PresentationTimeOrDelta = PresentationTimeDelta;
+			ar.Read(Parsed->EventDuration);
+			ar.Read(Parsed->ID);
+		}
+		else
+		{
+			ar.Read(Parsed->Timescale);
+			ar.Read(Parsed->PresentationTimeOrDelta);
+			ar.Read(Parsed->EventDuration);
+			ar.Read(Parsed->ID);
+			ar.ReadNullTerminatedStringUTF8(Parsed->SchemeIdUri);
+			ar.ReadNullTerminatedStringUTF8(Parsed->Value);
+		}
+		if (ar.GetNumBytesRemaining())
+		{
+			Parsed->MessageData.SetNumUninitialized(ar.GetNumBytesRemaining());
+			ar.ReadBytes(Parsed->MessageData.GetData(), ar.GetNumBytesRemaining());
+		}
+	}
+
+	uint8 FMP4BoxEMSG::GetVersion()
+	{
+		ParseIfRequired();
+		return Parsed->Version;
+	}
+	FString FMP4BoxEMSG::GetSchemeIdUri()
+	{
+		ParseIfRequired();
+		return Parsed->SchemeIdUri;
+	}
+	FString FMP4BoxEMSG::GetValue()
+	{
+		ParseIfRequired();
+		return Parsed->Value;
+	}
+	uint32 FMP4BoxEMSG::GetTimescale()
+	{
+		ParseIfRequired();
+		return Parsed->Timescale;
+	}
+	uint32 FMP4BoxEMSG::GetPresentationTimeDelta()
+	{
+		ParseIfRequired();
+		return (uint32)Parsed->PresentationTimeOrDelta;
+	}
+	uint64 FMP4BoxEMSG::GetPresentationTime()
+	{
+		ParseIfRequired();
+		return Parsed->PresentationTimeOrDelta;
+	}
+	uint32 FMP4BoxEMSG::GetEventDuration()
+	{
+		ParseIfRequired();
+		return Parsed->EventDuration;
+	}
+	uint32 FMP4BoxEMSG::GetID()
+	{
+		ParseIfRequired();
+		return Parsed->ID;
+	}
+	TConstArrayView<uint8> FMP4BoxEMSG::GetMessageData()
+	{
+		ParseIfRequired();
+		return Parsed->MessageData;
+	}
+
 /****************************************************************************************************************************************************/
 #if 0 // TEMPLATE
 	struct FMP4BoxCCCC::FParsed
@@ -2939,7 +3223,7 @@ namespace MP4Boxes
 	FMP4BoxFactory::FMP4BoxFactory()
 	{
 		FactoryMap.Add(MP4Utilities::MakeBoxAtom('f','t','y','p'), &FMP4BoxFTYP::Create);
-		FactoryMap.Add(MP4Utilities::MakeBoxAtom('s','t','y','p'), &FMP4BoxFTYP::Create);
+		FactoryMap.Add(MP4Utilities::MakeBoxAtom('s','t','y','p'), &FMP4BoxSTYP::Create);
 		FactoryMap.Add(MP4Utilities::MakeBoxAtom('m','o','o','v'), &FMP4BoxMOOV::Create);
 		FactoryMap.Add(MP4Utilities::MakeBoxAtom('m','v','h','d'), &FMP4BoxMVHD::Create);
 		FactoryMap.Add(MP4Utilities::MakeBoxAtom('t','r','a','k'), &FMP4BoxTRAK::Create);
@@ -2975,6 +3259,7 @@ namespace MP4Boxes
 		FactoryMap.Add(MP4Utilities::MakeBoxAtom('m','v','e','x'), &FMP4BoxMVEX::Create);
 		FactoryMap.Add(MP4Utilities::MakeBoxAtom('m','e','h','d'), &FMP4BoxMEHD::Create);
 		FactoryMap.Add(MP4Utilities::MakeBoxAtom('t','r','e','x'), &FMP4BoxTREX::Create);
+		FactoryMap.Add(MP4Utilities::MakeBoxAtom('p','r','f','t'), &FMP4BoxPRFT::Create);
 		FactoryMap.Add(MP4Utilities::MakeBoxAtom('l','e','v','a'), &FMP4BoxLEVA::Create);
 		FactoryMap.Add(MP4Utilities::MakeBoxAtom('m','o','o','f'), &FMP4BoxMOOF::Create);
 		FactoryMap.Add(MP4Utilities::MakeBoxAtom('m','f','h','d'), &FMP4BoxMFHD::Create);
@@ -3017,6 +3302,8 @@ namespace MP4Boxes
 		FactoryMap.Add(MP4Utilities::MakeBoxAtom('t','e','n','c'), &FMP4BoxTENC::Create);
 		FactoryMap.Add(MP4Utilities::MakeBoxAtom('s','e','n','c'), &FMP4BoxSENC::Create);
 
+		FactoryMap.Add(MP4Utilities::MakeBoxAtom('e','m','s','g'), &FMP4BoxEMSG::Create);
+
 		// `mdat` box is never parsed, we only handle it as a base box here to create the box tree.
 		FactoryMap.Add(MP4Utilities::MakeBoxAtom('m','d','a','t'), &FMP4BoxBase::Create);
 
diff --git a/Engine/Plugins/Media/MP4Utilities/Source/MP4Boxes/Private/MP4Track.cpp b/Engine/Plugins/Media/MP4Utilities/Source/MP4Boxes/Private/MP4Track.cpp
--- a/Engine/Plugins/Media/MP4Utilities/Source/MP4Boxes/Private/MP4Track.cpp
+++ b/Engine/Plugins/Media/MP4Utilities/Source/MP4Boxes/Private/MP4Track.cpp
@@ -18,6 +18,8 @@ namespace MP4Boxes
 		}
 		TkhdBox = InTrakBox->FindBoxRecursive<FMP4BoxTKHD>(MP4Utilities::MakeBoxAtom('t','k','h','d'), 0);
 		MdhdBox = InTrakBox->FindBoxRecursive<FMP4BoxMDHD>(MP4Utilities::MakeBoxAtom('m','d','h','d'), 1);
+		TencBox = InTrakBox->FindBoxRecursive<FMP4BoxTENC>(MP4Utilities::MakeBoxAtom('t','e','n','c'), 10);
+
 		auto _StszBox = InTrakBox->FindBoxRecursive<FMP4BoxSTSZ>(MP4Utilities::MakeBoxAtom('s','t','s','z'), 5);
 
 		check(TkhdBox && MdhdBox && _StszBox);
@@ -63,9 +65,19 @@ namespace MP4Boxes
 				Frg->TrafBox = TrackTrafBoxes[i];
 				Frg->MoofBox = Frg->TrafBox->FindParentBox<FMP4BoxMOOF>(MP4Utilities::MakeBoxAtom('m','o','o','f'));
 				Frg->TfhdBox = Frg->TrafBox->FindBoxRecursive<FMP4BoxTFHD>(MP4Utilities::MakeBoxAtom('t','f','h','d'), 0);
+				Frg->SencBox = Frg->TrafBox->FindBoxRecursive<FMP4BoxSENC>(MP4Utilities::MakeBoxAtom('s','e','n','c'), 0);
+				if (Frg->SencBox)
+				{
+					Frg->SencBox->Prepare(TencBox);
+				}
 				Frg->TrafBox->GetAllBoxInstances<FMP4BoxTRUN>(Frg->TrunBoxes, MP4Utilities::MakeBoxAtom('t','r','u','n'));
 				Frg->TrafBox->GetAllBoxInstances<FMP4BoxSUBS>(Frg->SubsBoxes, MP4Utilities::MakeBoxAtom('s','u','b','s'));
 				Frg->AbsoluteFirstSampleNumber = NumBaseTrackSamples + NumTotalFragmentSamples;
+				auto TfdtBox = Frg->TrafBox->FindBoxRecursive<FMP4BoxTFDT>(MP4Utilities::MakeBoxAtom('t','f','d','t'), 0);
+				if (TfdtBox)
+				{
+					Frg->BaseMediaDecodeTime = TfdtBox->GetBaseMediaDecodeTime();
+				}
 
 				Frg->NumSamples = 0;
 				for(int32 j=0, jMax=Frg->TrunBoxes.Num(); j<jMax; ++j)
@@ -121,7 +133,16 @@ namespace MP4Boxes
 						}
 					}
 				}
-				tf->FirstTrackTotalDTS = BaseTrackDuration + SumOfFragmentDurations;
+				if (tf->BaseMediaDecodeTime.IsSet())
+				{
+					//int64 ExpectedDTS = BaseTrackDuration + SumOfFragmentDurations;
+					int64 bmdt = (int64)tf->BaseMediaDecodeTime.GetValue();
+					tf->FirstTrackTotalDTS = bmdt;
+				}
+				else
+				{
+					tf->FirstTrackTotalDTS = BaseTrackDuration + SumOfFragmentDurations;
+				}
 				tf->LastTrackTotalDTSAndDur = tf->FirstTrackTotalDTS + tf->Duration;
 				SumOfFragmentDurations += tf->Duration;
 			}
@@ -248,6 +269,12 @@ namespace MP4Boxes
 		// Optional sample information boxes:
 		CttsBox = StblBox->FindBoxRecursive<FMP4BoxCTTS>(MP4Utilities::MakeBoxAtom('c','t','t','s'), 0);
 		StssBox = StblBox->FindBoxRecursive<FMP4BoxSTSS>(MP4Utilities::MakeBoxAtom('s','t','s','s'), 0);
+		TencBox = StblBox->FindBoxRecursive<FMP4BoxTENC>(MP4Utilities::MakeBoxAtom('t','e','n','c'), 6);
+		SencBox = TrakBox->FindBoxRecursive<FMP4BoxSENC>(MP4Utilities::MakeBoxAtom('s','e','n','c'), 1);
+		if (TencBox && SencBox)
+		{
+			SencBox->Prepare(TencBox);
+		}
 		StblBox->GetAllBoxInstances(SubsBoxes, MP4Utilities::MakeBoxAtom('s','u','b','s'));
 		StblBox->GetAllBoxInstances(SgpdBoxes, MP4Utilities::MakeBoxAtom('s','g','p','d'));
 		StblBox->GetAllBoxInstances(SbgpBoxes, MP4Utilities::MakeBoxAtom('s','b','g','p'));
@@ -286,54 +313,51 @@ namespace MP4Boxes
 				LastErrorMessage = FString::Printf(TEXT("Edit list is empty."));
 				return false;
 			}
-			if (ElstEntries.Num() > 1)
-			{
-				LastErrorMessage = FString::Printf(TEXT("Edit list with more than one entry is not supported."));
-				return false;
-			}
-			if (!(ElstEntries[0].media_rate_integer == 1 && ElstEntries[0].media_rate_fraction == 0))
-			{
-				LastErrorMessage = FString::Printf(TEXT("Edit list entries with playback rates other than 1.0 are not supported."));
-				return false;
-			}
-			if (ElstEntries[0].media_time < 0)
-			{
-				LastErrorMessage = FString::Printf(TEXT("Edit list specifies an empty edit, which is not supported."));
-				return false;
-			}
-
-			// The `media_time` in the edit list entry is specified in composition time. Typically an entry is used to shift
-			// a track having non-zero composition times to zero. In that case the `media_time` should correspond to the
-			// first composition time offset, or be greater.
-			// If it is less then the entry maps non-existing media, meaning that it essentially inserts an empty edit,
-			// which we do not support (what should be displayed then?)
-			if (ElstEntries[0].media_time < Convs.CompositionTimeAtZeroPoint)
+			// Look for the first entry that maps media time to establish the start offset.
+			// Any other entries we ignore.
+			int32 Idx=-1;
+			for(int32 i=0; i<ElstEntries.Num(); ++i)
 			{
-				// We assume that this is not actually wanted but a problem with the tool creating the file,
-				// so we emit a warning and ignore the edits `media_time`
-				UE_LOG(LogMP4Boxes, Verbose, TEXT("Edit list entry of track #%u maps non-existent media at composition time %lld to the timeline track. First available media composition time is %lld. Empty media will be ignored."), Convs.TrackID, (long long int)ElstEntries[0].media_time, (long long int)Convs.CompositionTimeAtZeroPoint);
+				if (ElstEntries[i].media_time >= 0)
+				{
+					Idx = i;
+					break;
+				}
 			}
-			else
+			if (Idx != -1)
 			{
-				Convs.CompositionTimeAtZeroPoint = ElstEntries[0].media_time;
-			}
+				// The `media_time` in the edit list entry is specified in composition time. Typically an entry is used to shift
+				// a track having non-zero composition times to zero. In that case the `media_time` should correspond to the
+				// first composition time offset, or be greater.
+				// If it is less then the entry maps non-existing media, meaning that it essentially inserts an empty edit,
+				// which we do not support (what should be displayed then?)
+				if (ElstEntries[Idx].media_time < Convs.CompositionTimeAtZeroPoint)
+				{
+					// We assume that this is not actually wanted but a problem with the tool creating the file,
+					// so we emit a warning and ignore the edits `media_time`
+					UE_LOG(LogMP4Boxes, Verbose, TEXT("Edit list entry of track #%u maps non-existent media at composition time %lld to the timeline track. First available media composition time is %lld. Empty media will be ignored."), Convs.TrackID, (long long int)ElstEntries[0].media_time, (long long int)Convs.CompositionTimeAtZeroPoint);
+				}
+				else
+				{
+					Convs.CompositionTimeAtZeroPoint = ElstEntries[Idx].media_time;
+				}
 
-			// The mapped duration may be different from the media duration itself, in which case the track is either truncated
-			// or the last sample has a longer duration to be repeated until the end.
-			if (ElstEntries[0].edit_duration == 0)
-			{
-				// The value of 0 is reserved for fragmented files with no `mehd` box.
-				// We do not handle this case.
-				if (!FragmentInfo.IsValid())
+				// The mapped duration may be different from the media duration itself, in which case the track is either truncated
+				// or the last sample has a longer duration to be repeated until the end.
+				if (ElstEntries[Idx].edit_duration == 0)
 				{
-					LastErrorMessage = FString::Printf(TEXT("Edit list specifies zero edit duration, which is not supported."));
-					return false;
+					// The value of 0 is reserved for fragmented files with no `mehd` box.
+					if (!FragmentInfo.IsValid())
+					{
+						LastErrorMessage = FString::Printf(TEXT("Edit list specifies zero edit duration, which is not supported."));
+						return false;
+					}
+					// Otherwise the mapped duration is already set as the entire movie duration.
+				}
+				else
+				{
+					Convs.MappedDurationFromElstBox.SetFromND(ElstEntries[Idx].edit_duration, Convs.DurationFromMvhdBox.GetDenominator());
 				}
-				// Otherwise the mapped duration is already set as the entire movie duration.
-			}
-			else
-			{
-				Convs.MappedDurationFromElstBox.SetFromND(ElstEntries[0].edit_duration, Convs.DurationFromMvhdBox.GetDenominator());
 			}
 		}
 		// For convenience sake convert the mapped duration from `mvhd` timescale into `mdhd` timescale.
@@ -383,12 +407,17 @@ namespace MP4Boxes
 
 		bHasBeenPrepared = true;
 
+		// If we are called with no movie duration we assume this is for iterating a `moof` of a fragmented
+		// file with a separate initialization segment. In this case we do not set up first and last sample
+		// convenience information as it is not needed and is not guaranteed to find sync samples.
+		if (!InFullMovieDuration.IsValid() && FragmentInfo.IsValid())
+		{
+			return true;
+		}
+
 		// Given the timeline mapping, locate the sample number that falls onto the start of the timeline
 		// and last one falling onto the end of the timeline.
 		TSharedPtr<FIterator, ESPMode::ThreadSafe> StartIt = CreateIterator();
-
-		// If we could not create an iterator then this file is most likely a fragmented mp4
-		// which is not handled.
 		if (!StartIt.IsValid())
 		{
 			LastErrorMessage = FString::Printf(TEXT("Could not locate first media sample. Is this an empty mp4?"));
@@ -556,6 +585,8 @@ namespace MP4Boxes
 			It->stssIt.SetBox(StssBox, Convs.NumMoovSamples);
 			// Are there subsample boxes? If yes, we use the first one only in the iterator for now.
 			It->subsIt.SetBox(SubsBoxes.Num() ? SubsBoxes[0] : nullptr, Convs.NumMoovSamples);
+			// Encryption?
+			It->sencIt.SetBox(SencBox, Convs.NumMoovSamples);
 			// Do we have a `rap ` group?
 			TSharedPtr<FMP4BoxSGPD, ESPMode::ThreadSafe>* rapSGPD = SgpdBoxes.FindByPredicate([](const TSharedPtr<FMP4BoxSGPD, ESPMode::ThreadSafe>& e){ return e->GetGroupingType() == MP4Utilities::MakeBoxAtom('r','a','p',' ');});
 			TSharedPtr<FMP4BoxSBGP, ESPMode::ThreadSafe>* rapSBGP = SbgpBoxes.FindByPredicate([](const TSharedPtr<FMP4BoxSBGP, ESPMode::ThreadSafe>& e){ return e->GetGroupingType() == MP4Utilities::MakeBoxAtom('r','a','p',' ');});
@@ -707,6 +738,7 @@ namespace MP4Boxes
 			stssIt.SetToSampleNumber(InFrameNumber);
 			rapIt.SetToSampleNumber(InFrameNumber);
 			subsIt.SetToSampleNumber(InFrameNumber);
+			sencIt.SetToSampleNumber(InFrameNumber);
 			FragmentVars.bInFragments = false;
 		}
 		else
@@ -726,7 +758,7 @@ namespace MP4Boxes
 		if (!FragmentVars.bInFragments)
 		{
 			// We are in the samples from the moov box.
-			int64 DTS = sttsIt.GetCurrentTime() + Convs.BaseMediaDecodeTime;
+			int64 DTS = sttsIt.GetCurrentTime();
 			uint32 Duration = sttsIt.GetCurrentDuration();
 			int64 CompositionTimeOffset = cttsIt.GetCurrentOffset();
 			int64 PTS = DTS + CompositionTimeOffset;
@@ -742,11 +774,14 @@ namespace MP4Boxes
 			// in timespan units. This is to avoid transformation issues from media local
 			// time into the timescale used in engine.
 			CurrentDurationTS = MP4Utilities::FFractionalTime(DTS + Duration, Timescale).GetAsTimespan() - CurrentDTS.GetAsTimespan();
+			CurrentBaseMediaDecodeTime.Reset();
 			CurrentSampleSize = stszIt.GetCurrentSampleSize();
 			bCurrentIsSyncOrRAP = stssIt.IsSyncSample() || rapIt.GetCurrentGroupDescriptionIndex() != 0;
 			// Subsample list
 			CurrentSubsampleList = subsIt.GetCurrentSubsamples();
 			CurrentSubsampleBoxFlags = subsIt.GetSubsampleBoxFlags();
+			// Encryption
+			EncryptionEntry = sencIt.GetCurrentEntry();
 
 			// Which chunk is this sample in?
 			uint32 chunkIndex = stscIt.GetCurrentChunkIndex();
@@ -766,9 +801,8 @@ namespace MP4Boxes
 		else
 		{
 			// We are in the fragment samples.
-			const int64 BaseMediaDecodeTime = 0;
-			int64 DTS = FragmentVars.TrafTrunIterator.GetDTS() + BaseMediaDecodeTime;
-			int64 PTS = FragmentVars.TrafTrunIterator.GetPTS() + BaseMediaDecodeTime;
+			int64 DTS = FragmentVars.TrafTrunIterator.GetDTS();
+			int64 PTS = FragmentVars.TrafTrunIterator.GetPTS();
 			int64 Duration = FragmentVars.TrafTrunIterator.GetDuration();
 			const uint32 Timescale = Convs.DurationFromMdhdBox.GetDenominator();
 			CurrentDTS.SetFromND(DTS, Timescale);
@@ -781,6 +815,14 @@ namespace MP4Boxes
 			// in timespan units. This is to avoid transformation issues from media local
 			// time into the timescale used in engine.
 			CurrentDurationTS = MP4Utilities::FFractionalTime(DTS + Duration, Timescale).GetAsTimespan() - CurrentDTS.GetAsTimespan();
+			if (FragmentVars.TrafTrunIterator.GetBaseMediaDecodeTime().IsSet())
+			{
+				CurrentBaseMediaDecodeTime = MP4Utilities::FFractionalTime((int64) FragmentVars.TrafTrunIterator.GetBaseMediaDecodeTime().GetValue(), Timescale);
+			}
+			else
+			{
+				CurrentBaseMediaDecodeTime.Reset();
+			}
 			CurrentSampleSize = FragmentVars.TrafTrunIterator.GetSampleSize();
 			bCurrentIsSyncOrRAP = FragmentVars.TrafTrunIterator.IsSyncSample();
 			CurrentSampleFileOffset = FragmentVars.TrafTrunIterator.GetAbsoluteOffset();
@@ -788,6 +830,8 @@ namespace MP4Boxes
 			// Subsample list
 			CurrentSubsampleList = FragmentVars.TrafTrunIterator.GetSubsampleList();
 			CurrentSubsampleBoxFlags = FragmentVars.TrafTrunIterator.GetSubsampleBoxFlags();
+			// Encryption
+			EncryptionEntry = FragmentVars.TrafTrunIterator.GetEncryptionInfo();
 		}
 	}
 
@@ -809,6 +853,7 @@ namespace MP4Boxes
 			verify(stssIt.Next());
 			verify(rapIt.Next());
 			verify(subsIt.Next());
+			verify(sencIt.Next());
 		}
 		else
 		{
@@ -847,6 +892,7 @@ namespace MP4Boxes
 			verify(stssIt.Prev());
 			verify(rapIt.Prev());
 			verify(subsIt.Prev());
+			verify(sencIt.Prev());
 		}
 		else
 		{
@@ -1003,6 +1049,10 @@ namespace MP4Boxes
 				// Are there subsample boxes? If yes, we use the first one only in the iterator for now.
 				subsIt.SetBox(Traf->SubsBoxes.Num() ? Traf->SubsBoxes[0] : nullptr, Traf->NumSamples);
 				subsIt.SetToSampleNumber(SampleNumber);
+
+				// Likewise there can only be a single `senc` box.
+				sencIt.SetBox(Traf->SencBox, Traf->NumSamples);
+				sencIt.SetToSampleNumber(SampleNumber);
 			}
 		}
 		check(Trun.IsValid());
@@ -1038,6 +1088,8 @@ namespace MP4Boxes
 			BaseOffset = InTraf->MoofBox->GetBoxFileOffset();
 		}
 
+		BaseMediaDecodeTime = InTraf->BaseMediaDecodeTime;
+
 		// Calculate the byte range of the first sample in the first `trun` up to and including the last sample in the last `trun`.
 		// This does not necessarily mean that all samples are consecutive in the file as each `trun` has its own offset so `trun`'s
 		// could as well be interleaved across sibling `traf`'s.
@@ -1104,6 +1156,8 @@ namespace MP4Boxes
 		}
 		// Subsample list
 		SubsampleList = subsIt.GetCurrentSubsamples();
+		// Encryption
+		EncryptionEntry = sencIt.GetCurrentEntry();
 	}
 
 	bool FMP4Track::FIterator::FTrafTrunIterator::Prev()
@@ -1117,6 +1171,8 @@ namespace MP4Boxes
 
 		// Step the subsample iterator. There is only one for all truns so we do not need to handle moving across multiple.
 		subsIt.Prev();
+		// Same for the encryption
+		sencIt.Prev();
 
 		TSharedPtr<FFragmentInfo::FTrackTrafs, ESPMode::ThreadSafe> Traf = FragmentInfo->TrackFragments[TrafNumber];
 		TSharedPtr<FMP4BoxTRUN, ESPMode::ThreadSafe> Trun = Traf->TrunBoxes[TrunNumber];
@@ -1133,6 +1189,8 @@ namespace MP4Boxes
 			PTS = DTS + (Trun->HasSampleCompositionTimeOffsets() ? Trun->GetSampleCompositionTimeOffsets()[SampleNumber] : 0);
 			// Update subsample list
 			SubsampleList = subsIt.GetCurrentSubsamples();
+			// Update encryption
+			EncryptionEntry = sencIt.GetCurrentEntry();
 			return true;
 		}
 		// Is there a trun before we can move into?
@@ -1183,6 +1241,10 @@ namespace MP4Boxes
 			subsIt.SetBox(Traf->SubsBoxes.Num() ? Traf->SubsBoxes[0] : nullptr, Traf->NumSamples);
 			subsIt.SetToSampleNumber(SampleNumber);
 
+			// Encryption?
+			sencIt.SetBox(Traf->SencBox, Traf->NumSamples);
+			sencIt.SetToSampleNumber(SampleNumber);
+
 			Update(Trun);
 			return true;
 		}
@@ -1198,6 +1260,8 @@ namespace MP4Boxes
 
 		// Advance the subsample iterator. There is only one for all truns so we do not need to handle moving across multiple.
 		subsIt.Next();
+		// Same for the encryption
+		sencIt.Next();
 
 		// Still samples in this trun?
 		if (SampleNumber + 1 < (int32)Trun->GetNumberOfSamples())
@@ -1212,6 +1276,8 @@ namespace MP4Boxes
 			Duration = Trun->HasSampleDurations() ? Trun->GetSampleDurations()[SampleNumber] : DefaultSampleDuration;
 			// Update subsample list
 			SubsampleList = subsIt.GetCurrentSubsamples();
+			// Update encryption
+			EncryptionEntry = sencIt.GetCurrentEntry();
 			return true;
 		}
 		// Is there another trun following we can move into?
@@ -1243,6 +1309,10 @@ namespace MP4Boxes
 			subsIt.SetBox(Traf->SubsBoxes.Num() ? Traf->SubsBoxes[0] : nullptr, Traf->NumSamples);
 			subsIt.SetToSampleNumber(SampleNumber);
 
+			// Encryption?
+			sencIt.SetBox(Traf->SencBox, Traf->NumSamples);
+			sencIt.SetToSampleNumber(SampleNumber);
+
 			Update(Trun);
 			return true;
 		}
diff --git a/Engine/Plugins/Media/MP4Utilities/Source/MP4Boxes/Public/MP4BoxBase.h b/Engine/Plugins/Media/MP4Utilities/Source/MP4Boxes/Public/MP4BoxBase.h
--- a/Engine/Plugins/Media/MP4Utilities/Source/MP4Boxes/Public/MP4BoxBase.h
+++ b/Engine/Plugins/Media/MP4Utilities/Source/MP4Boxes/Public/MP4BoxBase.h
@@ -200,11 +200,11 @@ namespace MP4Boxes
 	{
 	public:
 		UE_API bool ParseBoxTree(const TSharedPtr<MP4Utilities::FMP4BoxInfo, ESPMode::ThreadSafe>& InRootBox);
-		TSharedPtr<FMP4BoxBase> GetBoxTree() const
+		TSharedPtr<FMP4BoxBase, ESPMode::ThreadSafe> GetBoxTree() const
 		{ return BoxTree; }
 	private:
-		UE_API bool ParseBoxTreeInternal(const TWeakPtr<FMP4BoxBase>& InParent, const MP4Utilities::FMP4BoxInfo& InBox);
-		TSharedPtr<FMP4BoxBase> BoxTree;
+		UE_API bool ParseBoxTreeInternal(const TWeakPtr<FMP4BoxBase, ESPMode::ThreadSafe>& InParent, const MP4Utilities::FMP4BoxInfo& InBox);
+		TSharedPtr<FMP4BoxBase, ESPMode::ThreadSafe> BoxTree;
 	};
 
 } // namespace MP4Boxes
diff --git a/Engine/Plugins/Media/MP4Utilities/Source/MP4Boxes/Public/MP4BoxIterators.h b/Engine/Plugins/Media/MP4Utilities/Source/MP4Boxes/Public/MP4BoxIterators.h
--- a/Engine/Plugins/Media/MP4Utilities/Source/MP4Boxes/Public/MP4BoxIterators.h
+++ b/Engine/Plugins/Media/MP4Utilities/Source/MP4Boxes/Public/MP4BoxIterators.h
@@ -925,5 +925,63 @@ namespace MP4Boxes
 		int32 CurrentEntryIndex = 0;
 	};
 
+	/****************************************************************************************************************************************************/
+
+	struct FSENCBoxIterator
+	{
+		void SetBox(const TSharedPtr<FMP4BoxSENC, ESPMode::ThreadSafe>& InSencBox, uint32 InMaxSamples)
+		{
+			NumTotalSamples = InMaxSamples;
+			CurrentSampleNum = 0;
+			SencBox = InSencBox;
+			if (SencBox.IsValid())
+			{
+				NumBoxSamples = SencBox->GetEntries().Num();
+			}
+		}
+
+		void SetToSampleNumber(uint32 InSampleNum)
+		{
+			InSampleNum = InSampleNum < NumTotalSamples ? InSampleNum : NumTotalSamples-1;
+			CurrentSampleNum = InSampleNum;
+		}
+
+		bool Next()
+		{
+			if (IsLast())
+			{
+				return false;
+			}
+			++CurrentSampleNum;
+			return true;
+		}
+
+		bool Prev()
+		{
+			if (IsFirst())
+			{
+				return false;
+			}
+			--CurrentSampleNum;
+			return true;
+		}
+
+		bool IsFirst() const
+		{ return CurrentSampleNum == 0; }
+		bool IsLast() const
+		{ return CurrentSampleNum + 1 >= NumTotalSamples; }
+
+		const FMP4BoxSENC::FEntry& GetCurrentEntry() const
+		{
+			return CurrentSampleNum < NumBoxSamples ? SencBox->GetEntries()[CurrentSampleNum] : EmptyEntry;
+		}
+
+	private:
+		TSharedPtr<FMP4BoxSENC, ESPMode::ThreadSafe> SencBox;
+		FMP4BoxSENC::FEntry EmptyEntry;
+		uint32 NumTotalSamples = 0;
+		uint32 NumBoxSamples = 0;
+		uint32 CurrentSampleNum = 0;
+	};
 
 } // namespace MP4Boxes
diff --git a/Engine/Plugins/Media/MP4Utilities/Source/MP4Boxes/Public/MP4Boxes.h b/Engine/Plugins/Media/MP4Utilities/Source/MP4Boxes/Public/MP4Boxes.h
--- a/Engine/Plugins/Media/MP4Utilities/Source/MP4Boxes/Public/MP4Boxes.h
+++ b/Engine/Plugins/Media/MP4Utilities/Source/MP4Boxes/Public/MP4Boxes.h
@@ -34,8 +34,8 @@ namespace MP4Boxes
 
 
 	/**
-	 * `ftyp` box, `styp` box
-	 * ISO/IEC 14496-12:2022 - 4.3 File-type box & 8.16.2 Segment type box
+	 * `ftyp` box
+	 * ISO/IEC 14496-12:2022 - 4.3 File-type box
 	 */
 	class FMP4BoxFTYP : public FMP4BoxBase
 	{
@@ -43,10 +43,34 @@ namespace MP4Boxes
 		static TSharedPtr<FMP4BoxBase, ESPMode::ThreadSafe> Create(const TWeakPtr<FMP4BoxBase, ESPMode::ThreadSafe>& InParent, const MP4Utilities::FMP4BoxInfo& InBoxInfo)
 		{ return MakeShareable<FMP4BoxFTYP>(new FMP4BoxFTYP(InParent, InBoxInfo)); }
 		virtual ~FMP4BoxFTYP() = default;
+		UE_API uint32 GetMajorBrand();
+		UE_API uint32 GetMinorVersion();
+		UE_API TConstArrayView<uint32> GetCompatibleBrands();
+		UE_API TConstArrayView<uint32> GetAllBrands();
 	protected:
 		FMP4BoxFTYP(const TWeakPtr<FMP4BoxBase, ESPMode::ThreadSafe>& InParent, const MP4Utilities::FMP4BoxInfo& InBoxInfo)
 			: FMP4BoxBase(InParent, InBoxInfo)
 		{ }
+		void ParseIfRequired();
+		struct FParsed;
+		TPimplPtr<FParsed> Parsed;
+	};
+
+
+	/**
+	 * `styp` box
+	 * ISO/IEC 14496-12:2022 - 8.16.2 Segment type box
+	 */
+	class FMP4BoxSTYP : public FMP4BoxFTYP
+	{
+	public:
+		static TSharedPtr<FMP4BoxBase, ESPMode::ThreadSafe> Create(const TWeakPtr<FMP4BoxBase, ESPMode::ThreadSafe>& InParent, const MP4Utilities::FMP4BoxInfo& InBoxInfo)
+		{ return MakeShareable<FMP4BoxSTYP>(new FMP4BoxSTYP(InParent, InBoxInfo)); }
+		virtual ~FMP4BoxSTYP() = default;
+	protected:
+		FMP4BoxSTYP(const TWeakPtr<FMP4BoxBase, ESPMode::ThreadSafe>& InParent, const MP4Utilities::FMP4BoxInfo& InBoxInfo)
+			: FMP4BoxFTYP(InParent, InBoxInfo)
+		{ }
 	};
 
 
@@ -1035,7 +1059,11 @@ namespace MP4Boxes
 		static TSharedPtr<FMP4BoxBase, ESPMode::ThreadSafe> Create(const TWeakPtr<FMP4BoxBase, ESPMode::ThreadSafe>& InParent, const MP4Utilities::FMP4BoxInfo& InBoxInfo)
 		{ return MakeShareable<FMP4BoxSAIZ>(new FMP4BoxSAIZ(InParent, InBoxInfo)); }
 		virtual ~FMP4BoxSAIZ() = default;
-		UE_API void Test();
+		UE_API bool HasAuxInfoType();
+		UE_API uint32 GetAuxInfoType();
+		UE_API uint32 GetAuxInfoTypeParameter();
+		UE_API uint32 GetSampleCount();
+		UE_API uint8 GetSampleInfoSize(uint32 InForSampleNumber);
 	protected:
 		FMP4BoxSAIZ(const TWeakPtr<FMP4BoxBase, ESPMode::ThreadSafe>& InParent, const MP4Utilities::FMP4BoxInfo& InBoxInfo)
 			: FMP4BoxBase(InParent, InBoxInfo)
@@ -1056,7 +1084,11 @@ namespace MP4Boxes
 		static TSharedPtr<FMP4BoxBase, ESPMode::ThreadSafe> Create(const TWeakPtr<FMP4BoxBase, ESPMode::ThreadSafe>& InParent, const MP4Utilities::FMP4BoxInfo& InBoxInfo)
 		{ return MakeShareable<FMP4BoxSAIO>(new FMP4BoxSAIO(InParent, InBoxInfo)); }
 		virtual ~FMP4BoxSAIO() = default;
-		UE_API void Test();
+		UE_API bool HasAuxInfoType();
+		UE_API uint32 GetAuxInfoType();
+		UE_API uint32 GetAuxInfoTypeParameter();
+		UE_API uint32 GetEntryCount();
+		UE_API uint64 GetEntry(uint32 InForEntryNumber);
 	protected:
 		FMP4BoxSAIO(const TWeakPtr<FMP4BoxBase, ESPMode::ThreadSafe>& InParent, const MP4Utilities::FMP4BoxInfo& InBoxInfo)
 			: FMP4BoxBase(InParent, InBoxInfo)
@@ -1207,6 +1239,49 @@ namespace MP4Boxes
 	};
 
 
+	/**
+	 * `prft` box
+	 * ISO/IEC 14496-12:2022 - 8.16.5 Producer reference time box
+	 */
+	class FMP4BoxPRFT : public FMP4BoxBase
+	{
+	public:
+		static TSharedPtr<FMP4BoxBase, ESPMode::ThreadSafe> Create(const TWeakPtr<FMP4BoxBase, ESPMode::ThreadSafe>& InParent, const MP4Utilities::FMP4BoxInfo& InBoxInfo)
+		{ return MakeShareable<FMP4BoxPRFT>(new FMP4BoxPRFT(InParent, InBoxInfo)); }
+		virtual ~FMP4BoxPRFT() = default;
+
+		enum class EReferenceType
+		{
+			Undefined,
+			Encoder,		// Input or output. If distinction is needed, look at the flags.
+			Writer,			// Fragment written or finalized. If distinction is needed, look at the flags.
+			Captured,		// Flags == 8|16
+			Application		// Flags == 8 (also Flags == 16 despite this not being defined to be set alone)
+		};
+
+		UE_API uint32 GetFlags();
+		UE_API uint32 GetReferenceTrackID();
+		UE_API uint64 GetNTPTimestamp();
+		UE_API uint64 GetMediaTime();
+		EReferenceType GetReferenceType()
+		{
+			uint32 Flg = GetFlags();
+			if (Flg == 0 || Flg == 1) return EReferenceType::Encoder;
+			else if (Flg == 2 || Flg == 4) return EReferenceType::Writer;
+			else if (Flg == 8 || Flg == 16) return EReferenceType::Application;
+			else if (Flg == 24) return EReferenceType::Captured;
+			return EReferenceType::Undefined;
+		}
+	protected:
+		FMP4BoxPRFT(const TWeakPtr<FMP4BoxBase, ESPMode::ThreadSafe>& InParent, const MP4Utilities::FMP4BoxInfo& InBoxInfo)
+			: FMP4BoxBase(InParent, InBoxInfo)
+		{ }
+		void ParseIfRequired();
+		struct FParsed;
+		TPimplPtr<FParsed> Parsed;
+	};
+
+
 	/**
 	 * `moof` box
 	 * ISO/IEC 14496-12:2022 - 8.8.4 Movie fragment box
@@ -1425,7 +1500,7 @@ namespace MP4Boxes
 			uint32 SAPType : 3;
 			uint32 SAPDeltaTime : 28;
 		};
-		using FEntryList = TChunkedArray<FEntry, 3000>;
+		using FEntryList = TChunkedArray<FEntry, 1024>;
 		UE_API uint32 GetReferenceID();
 		UE_API uint32 GetTimescale();
 		UE_API uint64 GetEarliestPresentationTime();
@@ -1612,10 +1687,14 @@ namespace MP4Boxes
 		static TSharedPtr<FMP4BoxBase, ESPMode::ThreadSafe> Create(const TWeakPtr<FMP4BoxBase, ESPMode::ThreadSafe>& InParent, const MP4Utilities::FMP4BoxInfo& InBoxInfo)
 		{ return MakeShareable<FMP4BoxFRMA>(new FMP4BoxFRMA(InParent, InBoxInfo)); }
 		virtual ~FMP4BoxFRMA() = default;
+		UE_API uint32 GetOriginalFormat();
 	protected:
 		FMP4BoxFRMA(const TWeakPtr<FMP4BoxBase, ESPMode::ThreadSafe>& InParent, const MP4Utilities::FMP4BoxInfo& InBoxInfo)
 			: FMP4BoxBase(InParent, InBoxInfo)
 		{ }
+		void ParseIfRequired();
+		struct FParsed;
+		TPimplPtr<FParsed> Parsed;
 	};
 
 
@@ -1629,10 +1708,15 @@ namespace MP4Boxes
 		static TSharedPtr<FMP4BoxBase, ESPMode::ThreadSafe> Create(const TWeakPtr<FMP4BoxBase, ESPMode::ThreadSafe>& InParent, const MP4Utilities::FMP4BoxInfo& InBoxInfo)
 		{ return MakeShareable<FMP4BoxSCHM>(new FMP4BoxSCHM(InParent, InBoxInfo)); }
 		virtual ~FMP4BoxSCHM() = default;
+		UE_API uint32 GetSchemeType();
+		UE_API uint32 GetSchemeVersion();
 	protected:
 		FMP4BoxSCHM(const TWeakPtr<FMP4BoxBase, ESPMode::ThreadSafe>& InParent, const MP4Utilities::FMP4BoxInfo& InBoxInfo)
 			: FMP4BoxBase(InParent, InBoxInfo)
 		{ }
+		void ParseIfRequired();
+		struct FParsed;
+		TPimplPtr<FParsed> Parsed;
 	};
 
 
@@ -1715,7 +1799,7 @@ namespace MP4Boxes
 		{ return MakeShareable<FMP4BoxSENC>(new FMP4BoxSENC(InParent, InBoxInfo)); }
 		virtual ~FMP4BoxSENC() = default;
 
-		UE_API void Prepare(const TArray<TSharedPtr<FMP4BoxBase, ESPMode::ThreadSafe>>& InRelatedBoxes);
+		UE_API void Prepare(TWeakPtr<FMP4BoxTENC, ESPMode::ThreadSafe> InTENCBox);
 
 		struct FSubSample
 		{
@@ -1735,8 +1819,37 @@ namespace MP4Boxes
 		void ParseIfRequired();
 		struct FParsed;
 		TPimplPtr<FParsed> Parsed;
-		struct FInitRefs;
-		TPimplPtr<FInitRefs> InitRefs;
+		TWeakPtr<FMP4BoxTENC, ESPMode::ThreadSafe> ParentTENCBox;
+	};
+
+
+	/**
+	 * `emsg` box
+	 * ISO/IEC 23009-1:2022 - 5.10.3.3 Event message box
+	 */
+	class FMP4BoxEMSG : public FMP4BoxBase
+	{
+	public:
+		static TSharedPtr<FMP4BoxBase, ESPMode::ThreadSafe> Create(const TWeakPtr<FMP4BoxBase, ESPMode::ThreadSafe>& InParent, const MP4Utilities::FMP4BoxInfo& InBoxInfo)
+		{ return MakeShareable<FMP4BoxEMSG>(new FMP4BoxEMSG(InParent, InBoxInfo)); }
+		virtual ~FMP4BoxEMSG() = default;
+
+		UE_API uint8 GetVersion();	// version 0: GetPresentationTimeDelta(); version 1: GetPresentationTime()
+		UE_API FString GetSchemeIdUri();
+		UE_API FString GetValue();
+		UE_API uint32 GetTimescale();
+		UE_API uint32 GetPresentationTimeDelta();
+		UE_API uint64 GetPresentationTime();
+		UE_API uint32 GetEventDuration();
+		UE_API uint32 GetID();
+		UE_API TConstArrayView<uint8> GetMessageData();
+	protected:
+		FMP4BoxEMSG(const TWeakPtr<FMP4BoxBase, ESPMode::ThreadSafe>& InParent, const MP4Utilities::FMP4BoxInfo& InBoxInfo)
+			: FMP4BoxBase(InParent, InBoxInfo)
+		{ }
+		void ParseIfRequired();
+		struct FParsed;
+		TPimplPtr<FParsed> Parsed;
 	};
 
 
diff --git a/Engine/Plugins/Media/MP4Utilities/Source/MP4Boxes/Public/MP4Track.h b/Engine/Plugins/Media/MP4Utilities/Source/MP4Boxes/Public/MP4Track.h
--- a/Engine/Plugins/Media/MP4Utilities/Source/MP4Boxes/Public/MP4Track.h
+++ b/Engine/Plugins/Media/MP4Utilities/Source/MP4Boxes/Public/MP4Track.h
@@ -8,6 +8,7 @@
 #include "Containers/ArrayView.h"
 #include "Containers/UnrealString.h"
 #include "Misc/Timespan.h"
+#include "Misc/Optional.h"
 #include "MP4Boxes.h"
 #include "MP4BoxIterators.h"
 #include "MP4BoxMetadata.h"
@@ -45,7 +46,6 @@ namespace MP4Boxes
 			MP4Utilities::FFractionalTime DurationFromTkhdBox;
 			MP4Utilities::FFractionalTime DurationFromMdhdBox;
 			MP4Utilities::FFractionalTime MappedDurationFromElstBox;
-			int64 BaseMediaDecodeTime = 0;
 			int64 CompositionTimeAtZeroPoint = 0;
 			int64 DTSShiftAtZeroPoint = 0;
 			uint32 TrackID = 0;
@@ -62,8 +62,10 @@ namespace MP4Boxes
 				TSharedPtr<FMP4BoxMOOF, ESPMode::ThreadSafe> MoofBox;
 				TSharedPtr<FMP4BoxTRAF, ESPMode::ThreadSafe> TrafBox;
 				TSharedPtr<FMP4BoxTFHD, ESPMode::ThreadSafe> TfhdBox;
+				TSharedPtr<FMP4BoxSENC, ESPMode::ThreadSafe> SencBox;
 				TArray<TSharedPtr<FMP4BoxTRUN, ESPMode::ThreadSafe>> TrunBoxes;
 				TArray<TSharedPtr<FMP4BoxSUBS, ESPMode::ThreadSafe>> SubsBoxes;
+				TOptional<uint64> BaseMediaDecodeTime;
 				uint32 AbsoluteFirstSampleNumber = 0;
 				int64 FirstTrackTotalDTS = 0;
 				int64 LastTrackTotalDTSAndDur = 0;
@@ -74,6 +76,7 @@ namespace MP4Boxes
 			TSharedPtr<FMP4BoxTKHD, ESPMode::ThreadSafe> TkhdBox;
 			TSharedPtr<FMP4BoxMDHD, ESPMode::ThreadSafe> MdhdBox;
 			TSharedPtr<FMP4BoxTREX, ESPMode::ThreadSafe> TrexBox;
+			TSharedPtr<FMP4BoxTENC, ESPMode::ThreadSafe> TencBox;
 			uint32 NumTotalFragmentSamples = 0;
 
 			bool HasBeenPrepared() const
@@ -143,18 +146,18 @@ namespace MP4Boxes
 			// Returns the effective track's duration, as specified by an edit list. Timescale has been converted into `mdhd` timescale!
 			MP4Utilities::FFractionalTime GetEffectiveTrackDuration() const
 			{ return Convs.MappedDurationFromElstBox; }
+			// Returns the base media decode time for the current fragment.
+			TOptional<MP4Utilities::FFractionalTime> GetBaseMediaDecodeTime() const
+			{ return CurrentBaseMediaDecodeTime; }
 			// Returns the flags from the `subs` full box. Unset if there is no subsample box.
 			TOptional<uint32> GetSubsampleBoxFlags() const
 			{ return CurrentSubsampleBoxFlags; }
 			// Returns the list of subsamples for this sample, if there are any.
 			const FMP4BoxSUBS::FSubsampleList& GetSubsampleList() const
 			{ return CurrentSubsampleList; }
-			// Returns information about the first sample that is mapped to the 0-based timeline via `elst` box.
-			const FFirstSample& GetFirstSampleInfo() const
-			{ return Convs.FirstSample; }
-			// Returns information about the last sample that is mapped to the 0-based timeline via `elst` box.
-			const FLastSample& GetLastSampleInfo() const
-			{ return Convs.LastSample; }
+			// Returns the encryption information for this sample.
+			const FMP4BoxSENC::FEntry& GetEncryptionInfo() const
+			{ return EncryptionEntry; } 
 
 			// Advances this iterator to the next sample. Returns true if there is one, false if not.
 			// This iterates over the entire track, ignoring timeline mapping.
@@ -225,12 +228,17 @@ namespace MP4Boxes
 				{ return Duration; }
 				bool IsSyncSample() const
 				{ return (SampleFlags & 0x00010000U) == 0; }
+				TOptional<uint64> GetBaseMediaDecodeTime() const
+				{ return BaseMediaDecodeTime; }
 
 				TOptional<uint32> GetSubsampleBoxFlags() const
 				{ return subsIt.GetSubsampleBoxFlags(); }
 				const FMP4BoxSUBS::FSubsampleList& GetSubsampleList() const
 				{ return SubsampleList; }
 
+				const FMP4BoxSENC::FEntry& GetEncryptionInfo() const
+				{ return EncryptionEntry; } 
+
 				bool GetCurrentTrafRemainingSampleInfo(FChunkInfo& OutInfo) const;
 			private:
 				void Update(const TSharedPtr<FMP4BoxTRUN, ESPMode::ThreadSafe>& InTrun);
@@ -241,6 +249,7 @@ namespace MP4Boxes
 				uint32 TrexDefaultSampleSize = 0;
 				uint32 TrexDefaultSampleFlags = 0;
 
+				TOptional<uint64> BaseMediaDecodeTime;
 				int64 TrunStartDTS = 0;
 				int64 BaseOffset = 0;
 				int64 DataOffset = 0;
@@ -265,6 +274,9 @@ namespace MP4Boxes
 
 				FSUBSBoxIterator subsIt;
 				FMP4BoxSUBS::FSubsampleList SubsampleList;
+
+				FSENCBoxIterator sencIt;
+				FMP4BoxSENC::FEntry EncryptionEntry;
 			};
 
 			FIterator() = default;
@@ -285,6 +297,7 @@ namespace MP4Boxes
 			FSTCOBoxIterator stcoIt;
 			FSBGPBoxIterator rapIt;
 			FSUBSBoxIterator subsIt;
+			FSENCBoxIterator sencIt;
 			uint32 SampleNumber = 0;
 
 			struct FFragmentVars
@@ -299,12 +312,14 @@ namespace MP4Boxes
 			MP4Utilities::FFractionalTime CurrentEffectiveDTS;
 			MP4Utilities::FFractionalTime CurrentEffectivePTS;
 			MP4Utilities::FFractionalTime CurrentDuration;
+			TOptional<MP4Utilities::FFractionalTime> CurrentBaseMediaDecodeTime;
 			FTimespan CurrentDurationTS;
 			int64 CurrentSampleFileOffset = 0;
 			int64 CurrentSampleSize = 0;
 			bool bCurrentIsSyncOrRAP = false;
 			FMP4BoxSUBS::FSubsampleList CurrentSubsampleList;
 			TOptional<uint32> CurrentSubsampleBoxFlags;
+			FMP4BoxSENC::FEntry EncryptionEntry;
 		};
 
 		// Create an interator starting at the first sample.
@@ -356,6 +371,8 @@ namespace MP4Boxes
 		TSharedPtr<FMP4BoxSTSZ, ESPMode::ThreadSafe> StszBox;
 		TSharedPtr<FMP4BoxSTCO, ESPMode::ThreadSafe> StcoBox;
 		TSharedPtr<FMP4BoxSTSS, ESPMode::ThreadSafe> StssBox;
+		TSharedPtr<FMP4BoxTENC, ESPMode::ThreadSafe> TencBox;
+		TSharedPtr<FMP4BoxSENC, ESPMode::ThreadSafe> SencBox;
 		TSharedPtr<FMP4BoxUDTA, ESPMode::ThreadSafe> UdtaBox;
 
 		TArray<TSharedPtr<FMP4BoxSUBS, ESPMode::ThreadSafe>> SubsBoxes;
diff --git a/Engine/Plugins/Media/MP4Utilities/Source/MP4Utilities/Private/MP4Utilities.cpp b/Engine/Plugins/Media/MP4Utilities/Source/MP4Utilities/Private/MP4Utilities.cpp
--- a/Engine/Plugins/Media/MP4Utilities/Source/MP4Utilities/Private/MP4Utilities.cpp
+++ b/Engine/Plugins/Media/MP4Utilities/Source/MP4Utilities/Private/MP4Utilities.cpp
@@ -5,6 +5,15 @@
 
 namespace MP4Utilities
 {
+	TSharedPtr<FMP4BoxData, ESPMode::ThreadSafe> CreateBoxDataFromBuffer(TConstArrayView<uint8> InBoxDataBuffer, int64 InBoxDataFileOffset)
+	{
+		TSharedPtr<FMP4BoxData, ESPMode::ThreadSafe> bd = MakeShared<FMP4BoxData, ESPMode::ThreadSafe>();
+		bd->DataBuffer = InBoxDataBuffer;
+		FMP4AtomReader ar(bd->DataBuffer);
+		return ar.ParseIntoBoxInfo(*bd, InBoxDataFileOffset) ? bd : nullptr;
+	}
+
+
 	FMP4AtomReader::FMP4AtomReader(const TConstArrayView<uint8>& InData)
 		: DataPtr(InData.GetData()), DataSize(InData.Num()), CurrentOffset(0)
 	{
@@ -127,6 +136,38 @@ namespace MP4Utilities
 		return false;
 	}
 
+	bool FMP4AtomReader::ReadNullTerminatedStringUTF8(FString& OutString)
+	{
+		TArray<uint8> TempBuf;
+		TempBuf.Reserve(256);
+		while(GetNumBytesRemaining())
+		{
+			uint8 b;
+			if (!Read(b))
+			{
+				return false;
+			}
+			if (b == 0)
+			{
+				break;
+			}
+			TempBuf.Add(b);
+		}
+		if (TempBuf.Num())
+		{
+			// Check that this is not accidentally a Pascal string with the length (as a byte) as the first value.
+			if (TempBuf.Num() < 256 && TempBuf[0] == TempBuf.Num())
+			{
+				TempBuf.RemoveAt(0);
+			}
+			FUTF8ToTCHAR cnv((const ANSICHAR*)TempBuf.GetData(), TempBuf.Num());
+			OutString = FString::ConstructFromPtrSize(cnv.Get(), cnv.Length());
+			return true;
+		}
+		return false;
+	}
+
+
 	bool FMP4AtomReader::ReadAsNumber(uint64& OutValue, int32 InNumBytes)
 	{
 		OutValue = 0;
@@ -261,7 +302,7 @@ namespace MP4Utilities
 
 
 
-	bool FMP4BoxLocator::LocateAndReadRootBoxes(TArray<TSharedPtr<FMP4BoxData, ESPMode::ThreadSafe>>& OutBoxInfos, const TSharedPtr<IMP4DataReaderBase, ESPMode::ThreadSafe>& InDataReader, const TArray<uint32>& InFirstBoxes, const TArray<uint32>& InStopAfterBoxes, const TArray<uint32>& InReadDataOfBoxes, IMP4DataReaderBase::FCancellationCheckDelegate InCheckCancellationDelegate)
+	bool FMP4BoxLocator::LocateAndReadRootBoxes(TArray<TSharedPtr<FMP4BoxData, ESPMode::ThreadSafe>>& OutBoxInfos, const TSharedPtr<IMP4DataReaderBase, ESPMode::ThreadSafe>& InDataReader, const TArray<uint32>& InFirstBoxes, const TArray<uint32>& InStopAfterBoxes, bool bStopWithMDAT, const TArray<uint32>& InReadDataOfBoxes, IMP4DataReaderBase::FCancellationCheckDelegate InCheckCancellationDelegate)
 	{
 		// We NEVER want to read the `mdat` box here!
 		if (!InDataReader.IsValid() || InReadDataOfBoxes.Contains(MakeBoxAtom('m','d','a','t')))
@@ -339,15 +380,31 @@ namespace MP4Utilities
 				BoxInternalOffset += 16;
 			}
 			// Shall we read this box?
-			if ((bi->Type != MakeBoxAtom('m','d','a','t')) && (InReadDataOfBoxes.IsEmpty() || InReadDataOfBoxes.Contains(bi->Type)))
+			if (bi->Type != MakeBoxAtom('m','d','a','t'))
 			{
-				bi->DataBuffer.SetNumUninitialized(bi->Size-BoxInternalOffset);
-				bi->Data = MakeConstArrayView(bi->DataBuffer);
-				NumRead = InDataReader->ReadData(bi->DataBuffer.GetData(), bi->Size-BoxInternalOffset, CurrentOffset+BoxInternalOffset, InCheckCancellationDelegate);
+				if (InReadDataOfBoxes.IsEmpty() || InReadDataOfBoxes.Contains(bi->Type))
+				{
+					bi->DataBuffer.SetNumUninitialized(bi->Size-BoxInternalOffset);
+					bi->Data = MakeConstArrayView(bi->DataBuffer);
+					NumRead = InDataReader->ReadData(bi->DataBuffer.GetData(), bi->Size-BoxInternalOffset, CurrentOffset+BoxInternalOffset, InCheckCancellationDelegate);
+				}
+				else
+				{
+					NumRead = InDataReader->ReadData(nullptr, bi->Size-BoxInternalOffset, CurrentOffset+BoxInternalOffset, InCheckCancellationDelegate);
+				}
 			}
 			else
 			{
-				NumRead = InDataReader->ReadData(nullptr, bi->Size-BoxInternalOffset, CurrentOffset+BoxInternalOffset, InCheckCancellationDelegate);
+				if (!bStopWithMDAT)
+				{
+					NumRead = InDataReader->ReadData(nullptr, bi->Size-BoxInternalOffset, CurrentOffset+BoxInternalOffset, InCheckCancellationDelegate);
+				}
+				else
+				{
+					bi->DataOffset = BoxInternalOffset;
+					OutBoxInfos.Emplace(MoveTemp(bi));
+					return true;
+				}
 			}
 			CHECK_READ(bi->Size - BoxInternalOffset);
 			bi->DataOffset = BoxInternalOffset;
diff --git a/Engine/Plugins/Media/MP4Utilities/Source/MP4Utilities/Public/MP4Utilities.h b/Engine/Plugins/Media/MP4Utilities/Source/MP4Utilities/Public/MP4Utilities.h
--- a/Engine/Plugins/Media/MP4Utilities/Source/MP4Utilities/Public/MP4Utilities.h
+++ b/Engine/Plugins/Media/MP4Utilities/Source/MP4Utilities/Public/MP4Utilities.h
@@ -93,6 +93,16 @@ namespace MP4Utilities
 	#if !UE_BUILD_SHIPPING
 		char Name[5]{0};
 	#endif
+		TArray<uint8> GetBoxDataRAW() const
+		{
+			TArray<uint8> Raw;
+			Raw.AddUninitialized(8 + Data.Num());
+			uint32* bd = reinterpret_cast<uint32*>(Raw.GetData());
+			*bd++ = EndianSwap(static_cast<uint32>(Raw.Num()));
+			*bd++ = EndianSwap(Type);
+			FMemory::Memcpy(bd, Data.GetData(), Data.Num());
+			return Raw;
+		}
 	};
 
 	/**
@@ -105,6 +115,11 @@ namespace MP4Utilities
 	};
 
 
+	// Create an FMP4BoxData from a transient data buffer. The data is copied into the FMP4BoxData.
+	// If the box data fails to be recognized as an mp4 box a nullptr is returned.
+	UE_API TSharedPtr<FMP4BoxData, ESPMode::ThreadSafe> CreateBoxDataFromBuffer(TConstArrayView<uint8> InBoxDataBuffer, int64 InBoxDataFileOffset);
+
+
 	class FMP4AtomReader
 	{
 	public:
@@ -134,6 +149,7 @@ namespace MP4Utilities
 		UE_API bool ReadString(FString& OutString, uint16 InNumBytes);
 		UE_API bool ReadStringUTF8(FString& OutString, int32 InNumBytes);
 		UE_API bool ReadStringUTF16(FString& OutString, int32 InNumBytes);
+		UE_API bool ReadNullTerminatedStringUTF8(FString& OutString);
 		UE_API bool ReadBytes(void* OutBuffer, int32 InNumBytes);
 		UE_API bool ReadAsNumber(int64& OutValue, int32 InNumBytes);
 		UE_API bool ReadAsNumber(uint64& OutValue, int32 InNumBytes);
@@ -161,12 +177,12 @@ namespace MP4Utilities
 		FMP4BoxLocator() = default;
 		~FMP4BoxLocator() = default;
 
-		UE_API bool LocateAndReadRootBoxes(TArray<TSharedPtr<FMP4BoxData, ESPMode::ThreadSafe>>& OutBoxInfos, const TSharedPtr<IMP4DataReaderBase, ESPMode::ThreadSafe>& InDataReader, const TArray<uint32>& InFirstBoxes, const TArray<uint32>& InStopAfterBoxes, const TArray<uint32>& InReadDataOfBoxes, IMP4DataReaderBase::FCancellationCheckDelegate InCheckCancellationDelegate);
+		UE_API bool LocateAndReadRootBoxes(TArray<TSharedPtr<FMP4BoxData, ESPMode::ThreadSafe>>& OutBoxInfos, const TSharedPtr<IMP4DataReaderBase, ESPMode::ThreadSafe>& InDataReader, const TArray<uint32>& InFirstBoxes, const TArray<uint32>& InStopAfterBoxes, bool bStopWithMDAT, const TArray<uint32>& InReadDataOfBoxes, IMP4DataReaderBase::FCancellationCheckDelegate InCheckCancellationDelegate);
 		FString GetLastError() const
 		{ return LastError; }
 
 		static bool LocateAndReadRootBoxesFromBuffer(TArray<TSharedPtr<FMP4BoxData, ESPMode::ThreadSafe>>& OutBoxInfos, TConstArrayView<const uint8> InData)
-		{ return FMP4BoxLocator().LocateAndReadRootBoxes(OutBoxInfos, FMP4BufferDataReader::Create(InData), TArray<uint32>(), TArray<uint32>(), TArray<uint32>(), IMP4DataReaderBase::FCancellationCheckDelegate()); }
+		{ return FMP4BoxLocator().LocateAndReadRootBoxes(OutBoxInfos, FMP4BufferDataReader::Create(InData), TArray<uint32>(), TArray<uint32>(), false, TArray<uint32>(), IMP4DataReaderBase::FCancellationCheckDelegate()); }
 	private:
 		FString LastError;
 		int64 CurrentOffset = 0;
