Electra: Refactor mp4 track format determination to be shared by Electra and Protron

#rnx

#Synced-CL 51256620

diff --git a/Engine/Plugins/Media/ElectraCodecs/Source/ElectraDecoders/Private/Utils/MPEG/ElectraUtilsMP4.cpp b/Engine/Plugins/Media/ElectraCodecs/Source/ElectraDecoders/Private/Utils/MPEG/ElectraUtilsMP4.cpp
new file mode 100644
--- /dev/null
+++ b/Engine/Plugins/Media/ElectraCodecs/Source/ElectraDecoders/Private/Utils/MPEG/ElectraUtilsMP4.cpp
@@ -0,0 +1,373 @@
+// Copyright Epic Games, Inc. All Rights Reserved.
+
+#include "Utils/MPEG/ElectraUtilsMP4.h"
+#include "Features/IModularFeatures.h"
+#include "IElectraCodecFactoryModule.h"
+#include "IElectraCodecFactory.h"
+#include "IElectraCodecFormatUtils.h"
+#include "IElectraDecoderFeaturesAndOptions.h"
+#include "ElectraDecodersUtils.h"
+
+namespace ElectraDecodersUtil
+{
+	namespace MP4
+	{
+		bool GetTrackFormatInfo(FString& OutMessage, Electra::FCodecTypeFormat& OutCodecInfo, Electra::FDRMTypeFormat& OutDRMInfo, Electra::FDecoderInformation& OutDecoderInfo, const TSharedPtr<MP4Boxes::FMP4BoxTRAK, ESPMode::ThreadSafe>& InTrack, uint32 InTrackID, bool bInCheckIfDecodable)
+		{
+			IElectraCodecFormatUtilsModularFeature* CodecFormatUtils = nullptr;
+			{
+				IModularFeatures::FScopedLockModularFeatureList ScopedLockModularFeatureList;
+				TArray<IElectraCodecFormatUtilsModularFeature*> FormatUtils = IModularFeatures::Get().GetModularFeatureImplementations<IElectraCodecFormatUtilsModularFeature>(IElectraCodecFormatUtilsModularFeature::GetModularFeatureName());
+				CodecFormatUtils = !FormatUtils.IsEmpty() ? FormatUtils[0] : nullptr;
+			}
+
+			if (!CodecFormatUtils)
+			{
+				OutMessage = TEXT("No registered IElectraCodecFormatUtilsModularFeature found. Unable to determine codecs, ignoring track.");
+				return false;
+			}
+
+			IElectraCodecFactoryModule* CodecFactoryModule = nullptr;
+			if (bInCheckIfDecodable)
+			{
+				CodecFactoryModule = static_cast<IElectraCodecFactoryModule*>(FModuleManager::Get().GetModule(TEXT("ElectraCodecFactory")));
+				if (!CodecFactoryModule)
+				{
+					OutMessage = FString::Printf(TEXT("Electra decoder factory not found. Unable to use any track."));
+					return false;
+				}
+			}
+			auto CheckIfDecodable = [&]() -> void
+			{
+				// Shall we determine if there is a decoder able to decode this track?
+				if (CodecFactoryModule)
+				{
+					TMap<FString, FVariant> FormatInfo;
+					auto Factory = CodecFactoryModule->GetBestDecoderFactoryForFormat(FormatInfo, OutCodecInfo, TMap<FString, FVariant>());
+					if (Factory.IsValid())
+					{
+						OutDecoderInfo.bIsDecodable = true;
+						// Is there missing information the decoder factory could provide?
+						TMap<FString, FVariant> ConfigOptions;
+						Factory->GetConfigurationOptions(ConfigOptions);
+						if (OutCodecInfo.KeyframeMode == Electra::FCodecTypeFormat::EKeyframeMode::Unknown && ConfigOptions.Find(IElectraDecoderFormatInfo::IsEveryFrameKeyframe))
+						{
+							OutCodecInfo.KeyframeMode = !!ElectraDecodersUtil::GetVariantValueSafeI64(FormatInfo, IElectraDecoderFormatInfo::IsEveryFrameKeyframe, 0)
+														? Electra::FCodecTypeFormat::EKeyframeMode::OnlyKeyframes : Electra::FCodecTypeFormat::EKeyframeMode::DeltaFrames;
+						}
+						if (OutCodecInfo.HumanReadableFormatInfo.IsEmpty() && ConfigOptions.Find(IElectraDecoderFormatInfo::HumanReadableFormatName))
+						{
+							OutCodecInfo.HumanReadableFormatInfo = ElectraDecodersUtil::GetVariantValueFString(FormatInfo, IElectraDecoderFormatInfo::HumanReadableFormatName);
+						}
+						OutDecoderInfo.ProviderInfo.Name = Factory->GetProviderInformation().GetName();
+						OutDecoderInfo.ProviderInfo.Version = Factory->GetProviderInformation().GetVersion();
+						OutDecoderInfo.ProviderInfo.Implementation = Factory->GetProviderInformation().GetImplementation();
+						OutDecoderInfo.ProviderInfo.Vendor = Factory->GetProviderInformation().GetVendor();
+					}
+					else
+					{
+						OutDecoderInfo.bIsDecodable = false;
+					}
+				}
+			};
+
+			// There needs to be an `stsd` box in this track. We do not try the expected path of `trak`->`mdia`->`minf`->`stbl`->`stsd` as
+			// this is not that much faster and if the file is somewhat ill-formed we may not find it if it's grouped under elsewhere.
+			auto Stsd = InTrack->FindBoxRecursive<MP4Boxes::FMP4BoxSTSD>(MP4Utilities::MakeBoxAtom('s','t','s','d'), 6);
+			if (!Stsd)
+			{
+				// If not found we just ignore the track. That's a warning but not an error.
+				OutMessage = FString::Printf(TEXT("No `stsd` box found, ignoring track #%u."), InTrackID);
+				return false;
+			}
+
+			// Get language code
+			auto Mdhd = InTrack->FindBoxRecursive<MP4Boxes::FMP4BoxMDHD>(MP4Utilities::MakeBoxAtom('m','d','h','d'), 2);
+			if (Mdhd)
+			{
+				Electra::BCP47::ParseRFC5646Tag(OutCodecInfo.LanguageTag, Mdhd->GetLanguageCode639_2T());
+			}
+			auto Elng = InTrack->FindBoxRecursive<MP4Boxes::FMP4BoxELNG>(MP4Utilities::MakeBoxAtom('e','l','n','g'), 6);
+			if (Elng)
+			{
+				Electra::BCP47::ParseRFC5646Tag(OutCodecInfo.LanguageTag, Elng->GetLanguageTag());
+			}
+
+			auto AddChildren = [&OutCodecInfo](const TSharedPtr<MP4Boxes::FMP4BoxBase, ESPMode::ThreadSafe>& InSample) -> void
+			{
+				for(auto& Ch : InSample->GetChildren())
+				{
+					OutCodecInfo.ExtraBoxes.Emplace(Ch->GetType(), Ch->GetBoxData());
+				}
+			};
+
+			auto SetupEncryption = [&](uint32& InOutSampleType, const TSharedPtr<MP4Boxes::FMP4BoxBase, ESPMode::ThreadSafe>& InSample) -> bool
+			{
+				// Locate the encryption related boxes
+				auto Frma = InSample->FindBoxRecursive<MP4Boxes::FMP4BoxFRMA>(MP4Utilities::MakeBoxAtom('f','r','m','a'), 2);
+				if (!Frma)
+				{
+					OutMessage = FString::Printf(TEXT("Encrypted track #%u is missing mandatory `frma` box, ignoring this track."), InTrackID);
+					return false;
+				}
+				InOutSampleType = Frma->GetOriginalFormat();
+
+				OutDRMInfo.EncryptionInfo.Emplace<Electra::FDRMTypeFormat::FISOEncryptionInfo>();
+				Electra::FDRMTypeFormat::FISOEncryptionInfo& iso = OutDRMInfo.EncryptionInfo.Get<Electra::FDRMTypeFormat::FISOEncryptionInfo>();
+				// Optional `schm` box.
+				if (auto Schm = InSample->FindBoxRecursive<MP4Boxes::FMP4BoxSCHM>(MP4Utilities::MakeBoxAtom('s','c','h','m'), 2))
+				{
+					iso.Scheme = Schm->GetSchemeType();
+					iso.SchemeVersion = Schm->GetSchemeVersion();
+				}
+				// Optional `tenc` box.
+				if (auto Tenc = InSample->FindBoxRecursive<MP4Boxes::FMP4BoxTENC>(MP4Utilities::MakeBoxAtom('t','e','n','c'), 3))
+				{
+					if (Tenc->HasDefaultCryptBlockValues())
+					{
+						Electra::FDRMTypeFormat::FISOEncryptionInfo::FBlockPattern ptn;
+						ptn.CryptByteBlock = Tenc->GetDefaultCryptByteBlock();
+						ptn.SkipByteBlock = Tenc->GetDefaultSkipByteBlock();
+						iso.BlockPattern = MoveTemp(ptn);
+					}
+					iso.DefaultKID = Tenc->GetDefaultKID();
+					iso.DefaultIV = Tenc->GetDefaultConstantIV();
+					iso.DefaultIVSize = Tenc->GetDefaultPerSampleIVSize();
+				}
+				// Optional `pssh` boxes.
+				if (auto Moov = InSample->FindParentBox<MP4Boxes::FMP4BoxMOOV>(MP4Utilities::MakeBoxAtom('m','o','o','v')))
+				{
+					TArray<TSharedPtr<MP4Boxes::FMP4BoxPSSH, ESPMode::ThreadSafe>> PsshBoxes;
+					Moov->GetAllBoxInstances<MP4Boxes::FMP4BoxPSSH>(PsshBoxes, MP4Utilities::MakeBoxAtom('p','s','s','h'));
+					for(auto &pssh : PsshBoxes)
+					{
+						Electra::FDRMTypeFormat::FISOEncryptionInfo::FCDMInfo& cdm = iso.CDMInfos.Emplace_GetRef();
+						cdm.SystemID = pssh->GetSystemID();
+						cdm.Data = pssh->GetData();
+						for(auto &kid : pssh->GetKIDs())
+						{
+							cdm.KIDs.Emplace(kid);
+						}
+					}
+				}
+				return true;
+			};
+
+			// Get the sample type of this track. This call is necessary to trigger parsing of child nodes.
+			MP4Boxes::FMP4BoxSampleEntry::ESampleType SampleType = Stsd->GetSampleType();
+			// If already knows to not be supported, skip it.
+			if (SampleType == MP4Boxes::FMP4BoxSampleEntry::ESampleType::Unsupported)
+			{
+				OutMessage = FString::Printf(TEXT("Unsupported sample type, ignoring track #%u."), InTrackID);
+				return false;
+			}
+			// Several entries are permitted, but we need this to be unambiguous.
+			if (Stsd->GetChildren().Num() > 1)
+			{
+				OutMessage = FString::Printf(TEXT("Multiple sample descriptions found in `stsd` box, ignoring track #%u."), InTrackID);
+				return false;
+			}
+			else if (Stsd->GetChildren().Num() == 0)
+			{
+				OutMessage = FString::Printf(TEXT("No sample description found in `stsd` box, ignoring track #%u."), InTrackID);
+				return false;
+			}
+
+			OutDecoderInfo.bIsDecodable = true;
+			// Based on the sample type, get it and see if it is using codec we support.
+			if (SampleType == MP4Boxes::FMP4BoxSampleEntry::ESampleType::Video)
+			{
+				TArray<TSharedPtr<MP4Boxes::FMP4BoxVisualSampleEntry, ESPMode::ThreadSafe>> Visuals;
+				Stsd->GetSampleDescriptions(Visuals);
+				if (Visuals.Num() != 1)
+				{
+					if (Visuals.IsEmpty())
+					{
+						OutMessage = FString::Printf(TEXT("No visual sample entry found in `stsd` box, ignoring track #%u."), InTrackID);
+					}
+					else
+					{
+						OutMessage = FString::Printf(TEXT("More than one visual sample entry found in `stsd` box, ignoring track #%u."), InTrackID);
+					}
+					return false;
+				}
+
+				uint32 VisualType = Visuals[0]->GetType();
+				// Encrypted?
+				if (VisualType == MP4Utilities::MakeBoxAtom('e','n','c','v'))
+				{
+					if (!SetupEncryption(VisualType, Visuals[0]))
+					{
+						return false;
+					}
+				}
+				if (Visuals[0]->GetFrameCount() == 1)
+				{
+					OutCodecInfo.Type = Electra::FCodecTypeFormat::EType::Video;
+					OutCodecInfo.FourCC = VisualType;
+					OutCodecInfo.Properties.Emplace<Electra::FCodecTypeFormat::FVideo>();
+					Electra::FCodecTypeFormat::FVideo& cvid = OutCodecInfo.Properties.Get<Electra::FCodecTypeFormat::FVideo>();
+					cvid.Width = Visuals[0]->GetWidth();
+					cvid.Height = Visuals[0]->GetHeight();
+					auto Stts = InTrack->FindBoxRecursive<MP4Boxes::FMP4BoxSTTS>(MP4Utilities::MakeBoxAtom('s','t','t','s'), 5);
+					if (Mdhd.IsValid() && Stts.IsValid() && Stts->GetEntries().Num())
+					{
+						cvid.FrameRate = FFrameRate(Mdhd->GetTimescale(), Stts->GetEntries()[0].sample_delta);
+					}
+					for(auto& Ch : Visuals[0]->GetChildren())
+					{
+						OutCodecInfo.ExtraBoxes.Emplace(Ch->GetType(), Ch->GetBoxData());
+					}
+					if (!CodecFormatUtils->PrepareCodecTypeFormat(OutCodecInfo, IElectraCodecFormatUtilsModularFeature::FCodecTypeFormatParams()))
+					{
+						OutMessage = FString::Printf(TEXT("Failed to prepare the codec format type, ignoring track #%u."), InTrackID);
+						return false;
+					}
+					// Check if this video is decodable.
+					CheckIfDecodable();
+				}
+				else
+				{
+					OutMessage = FString::Printf(TEXT("Track #%u has a frame count other than 1 in the VisualSampleEntry, ignoring this track."), InTrackID);
+					return false;
+				}
+			}
+			else if (SampleType == MP4Boxes::FMP4BoxSampleEntry::ESampleType::Audio)
+			{
+				TArray<TSharedPtr<MP4Boxes::FMP4BoxAudioSampleEntry, ESPMode::ThreadSafe>> Audios;
+				Stsd->GetSampleDescriptions(Audios);
+				if (Audios.Num() != 1)
+				{
+					if (Audios.IsEmpty())
+					{
+						OutMessage = FString::Printf(TEXT("No audio sample entry found in `stsd` box, ignoring track #%u."), InTrackID);
+					}
+					else
+					{
+						OutMessage = FString::Printf(TEXT("More than one audio sample entry found in `stsd` box, ignoring track #%u."), InTrackID);
+					}
+					return false;
+				}
+
+				uint32 AudioType = Audios[0]->GetType();
+				// Encrypted?
+				if (AudioType == MP4Utilities::MakeBoxAtom('e','n','c','a'))
+				{
+					if (!SetupEncryption(AudioType, Audios[0]))
+					{
+						return false;
+					}
+				}
+				OutCodecInfo.Type = Electra::FCodecTypeFormat::EType::Audio;
+				OutCodecInfo.FourCC = AudioType;
+				OutCodecInfo.Properties.Emplace<Electra::FCodecTypeFormat::FAudio>();
+				Electra::FCodecTypeFormat::FAudio& caud = OutCodecInfo.Properties.Get<Electra::FCodecTypeFormat::FAudio>();
+				caud.NumChannels = Audios[0]->GetChannelCount();
+				caud.SampleRate = Audios[0]->GetSampleRate();
+				for(auto& Ch : Audios[0]->GetChildren())
+				{
+					OutCodecInfo.ExtraBoxes.Emplace(Ch->GetType(), Ch->GetBoxData());
+				}
+				if (!CodecFormatUtils->PrepareCodecTypeFormat(OutCodecInfo, IElectraCodecFormatUtilsModularFeature::FCodecTypeFormatParams()))
+				{
+					OutMessage = FString::Printf(TEXT("Failed to prepare the codec format type, ignoring track #%u."), InTrackID);
+					return false;
+				}
+				// Check if this audio is decodable.
+				CheckIfDecodable();
+			}
+			else if (SampleType == MP4Boxes::FMP4BoxSampleEntry::ESampleType::Subtitles || SampleType == MP4Boxes::FMP4BoxSampleEntry::ESampleType::TimedText)
+			{
+				TArray<TSharedPtr<MP4Boxes::FMP4BoxSampleEntry, ESPMode::ThreadSafe>> Subtitles;
+				Stsd->GetSampleDescriptions(Subtitles);
+				if (Subtitles.Num() != 1)
+				{
+					if (Subtitles.IsEmpty())
+					{
+						OutMessage = FString::Printf(TEXT("No subtitle sample entry found in `stsd` box, ignoring track #%u."), InTrackID);
+					}
+					else
+					{
+						OutMessage = FString::Printf(TEXT("More than one subtitle sample entry found in `stsd` box, ignoring track #%u."), InTrackID);
+					}
+					return false;
+				}
+
+				uint32 SubtitleType = Subtitles[0]->GetType();
+				// Encrypted?
+				if (SubtitleType == MP4Utilities::MakeBoxAtom('e','n','c','t') || SubtitleType == MP4Utilities::MakeBoxAtom('e','n','c','u'))
+				{
+					if (!SetupEncryption(SubtitleType, Subtitles[0]))
+					{
+						return false;
+					}
+				}
+				OutCodecInfo.Type = Electra::FCodecTypeFormat::EType::Subtitle;
+				OutCodecInfo.FourCC = SubtitleType;
+				OutCodecInfo.Properties.Emplace<Electra::FCodecTypeFormat::FSubtitle>();
+				//Electra::FCodecTypeFormat::FSubtitle& csub = OutCodecInfo.Properties.Get<Electra::FCodecTypeFormat::FSubtitle>();
+				// Add the sample box itself as an extra box.
+				OutCodecInfo.ExtraBoxes.Emplace(SubtitleType, Subtitles[0]->GetBoxData());
+				for(auto& Ch : Subtitles[0]->GetChildren())
+				{
+					OutCodecInfo.ExtraBoxes.Emplace(Ch->GetType(), Ch->GetBoxData());
+				}
+				if (!CodecFormatUtils->PrepareCodecTypeFormat(OutCodecInfo, IElectraCodecFormatUtilsModularFeature::FCodecTypeFormatParams()))
+				{
+					OutMessage = FString::Printf( TEXT("Failed to prepare the codec format type, ignoring track #%u."), InTrackID);
+					return false;
+				}
+			}
+			else if (SampleType == MP4Boxes::FMP4BoxSampleEntry::ESampleType::QTFFTimecode)
+			{
+				TArray<TSharedPtr<MP4Boxes::FMP4BoxQTFFTimecodeSampleEntry, ESPMode::ThreadSafe>> Timecodes;
+				Stsd->GetSampleDescriptions(Timecodes);
+				if (Timecodes.Num() != 1)
+				{
+					if (Timecodes.IsEmpty())
+					{
+						OutMessage = FString::Printf(TEXT("No timecode sample entry found in `stsd` box, ignoring track #%u."), InTrackID);
+					}
+					else
+					{
+						OutMessage = FString::Printf(TEXT("More than one timecode sample entry found in `stsd` box, ignoring track #%u."), InTrackID);
+					}
+					return false;
+				}
+
+				switch(Timecodes[0]->GetType())
+				{
+					case MP4Utilities::MakeBoxAtom('t','m','c','d'):
+					{
+						Electra::FCodecTypeFormat::FTMCDTimecode Timecode;
+
+						Timecode.Flags = Timecodes[0]->GetFlags();
+						Timecode.Timescale = Timecodes[0]->GetTimescale();
+						Timecode.FrameDuration = Timecodes[0]->GetFrameDuration();
+						Timecode.NumberOfFrames = Timecodes[0]->GetNumberOfFrames();
+						OutCodecInfo.Properties.Emplace<Electra::FCodecTypeFormat::FTMCDTimecode>(MoveTemp(Timecode));
+						OutCodecInfo.RFC6381 = MP4Utilities::GetPrintableBoxAtom(Timecodes[0]->GetType());
+						OutCodecInfo.FourCC = Timecodes[0]->GetType();
+						OutCodecInfo.DCR = Timecodes[0]->GetBoxData();
+						OutCodecInfo.Type = Electra::FCodecTypeFormat::EType::Timecode;
+						AddChildren(Timecodes[0]);
+						break;
+					}
+				}
+			}
+
+			// If the track has not been identified as usable so far, remove it.
+			if (OutCodecInfo.Type == Electra::FCodecTypeFormat::EType::Invalid)
+			{
+				TArray<TSharedPtr<MP4Boxes::FMP4BoxSampleEntry, ESPMode::ThreadSafe>> Entries;
+				Stsd->GetSampleDescriptions(Entries);
+				check(Entries.Num() == 1);
+				OutMessage = FString::Printf(TEXT("Track of sample type \"%s\" is not supported, ignoring track #%u."), *MP4Utilities::GetPrintableBoxAtom(Entries[0]->GetType()), InTrackID);
+				return false;
+			}
+			return true;
+		}
+
+	} // namespace MP4
+
+} // namespace ElectraDecodersUtil
diff --git a/Engine/Plugins/Media/ElectraCodecs/Source/ElectraDecoders/Public/Utils/MPEG/ElectraUtilsMP4.h b/Engine/Plugins/Media/ElectraCodecs/Source/ElectraDecoders/Public/Utils/MPEG/ElectraUtilsMP4.h
new file mode 100644
--- /dev/null
+++ b/Engine/Plugins/Media/ElectraCodecs/Source/ElectraDecoders/Public/Utils/MPEG/ElectraUtilsMP4.h
@@ -0,0 +1,18 @@
+// Copyright Epic Games, Inc. All Rights Reserved.
+
+#pragma once
+
+#include "CoreTypes.h"
+#include "Templates/SharedPointer.h"
+#include "CodecTypeFormat.h"
+#include "MP4Boxes.h"
+
+namespace ElectraDecodersUtil
+{
+	namespace MP4
+	{
+		ELECTRADECODERS_API bool GetTrackFormatInfo(FString& OutMessage, Electra::FCodecTypeFormat& OutCodecInfo, Electra::FDRMTypeFormat& OutDRMInfo, Electra::FDecoderInformation& OutDecoderInfo, const TSharedPtr<MP4Boxes::FMP4BoxTRAK, ESPMode::ThreadSafe>& InTrack, uint32 InTrackID, bool bInCheckIfDecodable);
+
+	} // namespace MP4
+
+} // namespace ElectraDecodersUtil
diff --git a/Engine/Plugins/Media/ElectraPlayer/Source/ElectraPlayerRuntime/Private/Runtime/Player/StreamSegmentReaderCommon.h b/Engine/Plugins/Media/ElectraPlayer/Source/ElectraPlayerRuntime/Private/Runtime/Player/StreamSegmentReaderCommon.h
--- a/Engine/Plugins/Media/ElectraPlayer/Source/ElectraPlayerRuntime/Private/Runtime/Player/StreamSegmentReaderCommon.h
+++ b/Engine/Plugins/Media/ElectraPlayer/Source/ElectraPlayerRuntime/Private/Runtime/Player/StreamSegmentReaderCommon.h
@@ -8,6 +8,7 @@
 #include "MP4Boxes.h"
 #include "MP4Utilities.h"
 #include "Utilities/MP4Helpers.h"
+#include "Utils/MPEG/ElectraUtilsMP4.h"
 #include "Demuxer/ParserISO14496-12.h"
 #include "Demuxer/ParserMKV.h"
 #include "HTTP/HTTPManager.h"
@@ -104,7 +105,7 @@ public:
 			ti.TrackTimescale = Mdhd->GetTimescale();
 
 			FString Msg;
-			if (!GetTrackFormatInfo(Msg, ti.CodecFormat, ti.DrmFormat, ti.DecoderInformation, AllTracks[nTk], ti.TrackID, true))
+			if (!ElectraDecodersUtil::MP4::GetTrackFormatInfo(Msg, ti.CodecFormat, ti.DrmFormat, ti.DecoderInformation, AllTracks[nTk], ti.TrackID, true))
 			{
 				continue;
 			}
diff --git a/Engine/Plugins/Media/ElectraPlayer/Source/ElectraPlayerRuntime/Private/Runtime/Player/isobmff/PlaylistISOBMFF.cpp b/Engine/Plugins/Media/ElectraPlayer/Source/ElectraPlayerRuntime/Private/Runtime/Player/isobmff/PlaylistISOBMFF.cpp
--- a/Engine/Plugins/Media/ElectraPlayer/Source/ElectraPlayerRuntime/Private/Runtime/Player/isobmff/PlaylistISOBMFF.cpp
+++ b/Engine/Plugins/Media/ElectraPlayer/Source/ElectraPlayerRuntime/Private/Runtime/Player/isobmff/PlaylistISOBMFF.cpp
@@ -12,6 +12,7 @@
 #include "Utilities/TimeUtilities.h"
 #include "Utilities/BCP47-Helpers.h"
 #include "Utilities/MP4Helpers.h"
+#include "Utils/MPEG/ElectraUtilsMP4.h"
 #include "ElectraPlayerPrivate.h"
 #include "HTTP/HTTPManager.h"
 #include "Misc/FrameRate.h"
@@ -1277,7 +1278,7 @@ FErrorDetail FManifestISOBMFFInternal::FTimelineAssetISOBMFF::Build(IPlayerSessi
 		*/
 
 		FString Msg;
-		if (!GetTrackFormatInfo(Msg, TrkInf->CodecFormat, TrkInf->DrmFormat, TrkInf->DecoderInformation, ThisTrakBox, TrkInf->TrackID, true))
+		if (!ElectraDecodersUtil::MP4::GetTrackFormatInfo(Msg, TrkInf->CodecFormat, TrkInf->DrmFormat, TrkInf->DecoderInformation, ThisTrakBox, TrkInf->TrackID, true))
 		{
 			continue;
 		}
diff --git a/Engine/Plugins/Media/ElectraPlayer/Source/ElectraPlayerRuntime/Private/Runtime/Utilities/MP4Helpers.cpp b/Engine/Plugins/Media/ElectraPlayer/Source/ElectraPlayerRuntime/Private/Runtime/Utilities/MP4Helpers.cpp
--- a/Engine/Plugins/Media/ElectraPlayer/Source/ElectraPlayerRuntime/Private/Runtime/Utilities/MP4Helpers.cpp
+++ b/Engine/Plugins/Media/ElectraPlayer/Source/ElectraPlayerRuntime/Private/Runtime/Utilities/MP4Helpers.cpp
@@ -4,12 +4,6 @@
 #include "MP4Utilities.h"
 #include "Utilities/URLParser.h"
 #include "Utilities/BCP47-Helpers.h"
-#include "Features/IModularFeatures.h"
-#include "IElectraCodecFactoryModule.h"
-#include "IElectraCodecFactory.h"
-#include "IElectraCodecFormatUtils.h"
-#include "IElectraDecoderFeaturesAndOptions.h"
-#include "ElectraDecodersUtils.h"
 
 namespace Electra
 {
@@ -241,364 +235,4 @@ bool FMP4DataLoader::HasReachedEOF()
 	return bEOS;
 }
 
-
-
-
-bool GetTrackFormatInfo(FString& OutMessage, FCodecTypeFormat& OutCodecInfo, FDRMTypeFormat& OutDRMInfo, FDecoderInformation& OutDecoderInfo, const TSharedPtr<MP4Boxes::FMP4BoxTRAK, ESPMode::ThreadSafe>& InTrack, uint32 InTrackID, bool bInCheckIfDecodable)
-{
-	IElectraCodecFormatUtilsModularFeature* CodecFormatUtils = nullptr;
-	{
-		IModularFeatures::FScopedLockModularFeatureList ScopedLockModularFeatureList;
-		TArray<IElectraCodecFormatUtilsModularFeature*> FormatUtils = IModularFeatures::Get().GetModularFeatureImplementations<IElectraCodecFormatUtilsModularFeature>(IElectraCodecFormatUtilsModularFeature::GetModularFeatureName());
-		CodecFormatUtils = !FormatUtils.IsEmpty() ? FormatUtils[0] : nullptr;
-	}
-
-	if (!CodecFormatUtils)
-	{
-		OutMessage = TEXT("No registered IElectraCodecFormatUtilsModularFeature found. Unable to determine codecs, ignoring track.");
-		return false;
-	}
-
-	IElectraCodecFactoryModule* CodecFactoryModule = nullptr;
-	if (bInCheckIfDecodable)
-	{
-		CodecFactoryModule = static_cast<IElectraCodecFactoryModule*>(FModuleManager::Get().GetModule(TEXT("ElectraCodecFactory")));
-		if (!CodecFactoryModule)
-		{
-			OutMessage = FString::Printf(TEXT("Electra decoder factory not found. Unable to use any track."));
-			return false;
-		}
-	}
-	auto CheckIfDecodable = [&]() -> void
-	{
-		// Shall we determine if there is a decoder able to decode this track?
-		if (CodecFactoryModule)
-		{
-			TMap<FString, FVariant> FormatInfo;
-			auto Factory = CodecFactoryModule->GetBestDecoderFactoryForFormat(FormatInfo, OutCodecInfo, TMap<FString, FVariant>());
-			if (Factory.IsValid())
-			{
-				OutDecoderInfo.bIsDecodable = true;
-				// Is there missing information the decoder factory could provide?
-				TMap<FString, FVariant> ConfigOptions;
-				Factory->GetConfigurationOptions(ConfigOptions);
-				if (OutCodecInfo.KeyframeMode == FCodecTypeFormat::EKeyframeMode::Unknown && ConfigOptions.Find(IElectraDecoderFormatInfo::IsEveryFrameKeyframe))
-				{
-					OutCodecInfo.KeyframeMode = !!ElectraDecodersUtil::GetVariantValueSafeI64(FormatInfo, IElectraDecoderFormatInfo::IsEveryFrameKeyframe, 0)
-												? FCodecTypeFormat::EKeyframeMode::OnlyKeyframes : FCodecTypeFormat::EKeyframeMode::DeltaFrames;
-				}
-				if (OutCodecInfo.HumanReadableFormatInfo.IsEmpty() && ConfigOptions.Find(IElectraDecoderFormatInfo::HumanReadableFormatName))
-				{
-					OutCodecInfo.HumanReadableFormatInfo = ElectraDecodersUtil::GetVariantValueFString(FormatInfo, IElectraDecoderFormatInfo::HumanReadableFormatName);
-				}
-				OutDecoderInfo.ProviderInfo.Name = Factory->GetProviderInformation().GetName();
-				OutDecoderInfo.ProviderInfo.Version = Factory->GetProviderInformation().GetVersion();
-				OutDecoderInfo.ProviderInfo.Implementation = Factory->GetProviderInformation().GetImplementation();
-				OutDecoderInfo.ProviderInfo.Vendor = Factory->GetProviderInformation().GetVendor();
-			}
-			else
-			{
-				OutDecoderInfo.bIsDecodable = false;
-			}
-		}
-	};
-
-	// There needs to be an `stsd` box in this track. We do not try the expected path of `trak`->`mdia`->`minf`->`stbl`->`stsd` as
-	// this is not that much faster and if the file is somewhat ill-formed we may not find it if it's grouped under elsewhere.
-	auto Stsd = InTrack->FindBoxRecursive<MP4Boxes::FMP4BoxSTSD>(MP4Utilities::MakeBoxAtom('s','t','s','d'), 6);
-	if (!Stsd)
-	{
-		// If not found we just ignore the track. That's a warning but not an error.
-		OutMessage = FString::Printf(TEXT("No `stsd` box found, ignoring track #%u."), InTrackID);
-		return false;
-	}
-
-	// Get language code
-	auto Mdhd = InTrack->FindBoxRecursive<MP4Boxes::FMP4BoxMDHD>(MP4Utilities::MakeBoxAtom('m','d','h','d'), 2);
-	if (Mdhd)
-	{
-		Electra::BCP47::ParseRFC5646Tag(OutCodecInfo.LanguageTag, Mdhd->GetLanguageCode639_2T());
-	}
-	auto Elng = InTrack->FindBoxRecursive<MP4Boxes::FMP4BoxELNG>(MP4Utilities::MakeBoxAtom('e','l','n','g'), 6);
-	if (Elng)
-	{
-		Electra::BCP47::ParseRFC5646Tag(OutCodecInfo.LanguageTag, Elng->GetLanguageTag());
-	}
-
-	auto AddChildren = [&OutCodecInfo](const TSharedPtr<MP4Boxes::FMP4BoxBase, ESPMode::ThreadSafe>& InSample) -> void
-	{
-		for(auto& Ch : InSample->GetChildren())
-		{
-			OutCodecInfo.ExtraBoxes.Emplace(Ch->GetType(), Ch->GetBoxData());
-		}
-	};
-
-	auto SetupEncryption = [&](uint32& InOutSampleType, const TSharedPtr<MP4Boxes::FMP4BoxBase, ESPMode::ThreadSafe>& InSample) -> bool
-	{
-		// Locate the encryption related boxes
-		auto Frma = InSample->FindBoxRecursive<MP4Boxes::FMP4BoxFRMA>(MP4Utilities::MakeBoxAtom('f','r','m','a'), 2);
-		if (!Frma)
-		{
-			OutMessage = FString::Printf(TEXT("Encrypted track #%u is missing mandatory `frma` box, ignoring this track."), InTrackID);
-			return false;
-		}
-		InOutSampleType = Frma->GetOriginalFormat();
-
-		OutDRMInfo.EncryptionInfo.Emplace<FDRMTypeFormat::FISOEncryptionInfo>();
-		FDRMTypeFormat::FISOEncryptionInfo& iso = OutDRMInfo.EncryptionInfo.Get<FDRMTypeFormat::FISOEncryptionInfo>();
-		// Optional `schm` box.
-		if (auto Schm = InSample->FindBoxRecursive<MP4Boxes::FMP4BoxSCHM>(MP4Utilities::MakeBoxAtom('s','c','h','m'), 2))
-		{
-			iso.Scheme = Schm->GetSchemeType();
-			iso.SchemeVersion = Schm->GetSchemeVersion();
-		}
-		// Optional `tenc` box.
-		if (auto Tenc = InSample->FindBoxRecursive<MP4Boxes::FMP4BoxTENC>(MP4Utilities::MakeBoxAtom('t','e','n','c'), 3))
-		{
-			if (Tenc->HasDefaultCryptBlockValues())
-			{
-				FDRMTypeFormat::FISOEncryptionInfo::FBlockPattern ptn;
-				ptn.CryptByteBlock = Tenc->GetDefaultCryptByteBlock();
-				ptn.SkipByteBlock = Tenc->GetDefaultSkipByteBlock();
-				iso.BlockPattern = MoveTemp(ptn);
-			}
-			iso.DefaultKID = Tenc->GetDefaultKID();
-			iso.DefaultIV = Tenc->GetDefaultConstantIV();
-			iso.DefaultIVSize = Tenc->GetDefaultPerSampleIVSize();
-		}
-		// Optional `pssh` boxes.
-		if (auto Moov = InSample->FindParentBox<MP4Boxes::FMP4BoxMOOV>(MP4Utilities::MakeBoxAtom('m','o','o','v')))
-		{
-			TArray<TSharedPtr<MP4Boxes::FMP4BoxPSSH, ESPMode::ThreadSafe>> PsshBoxes;
-			Moov->GetAllBoxInstances<MP4Boxes::FMP4BoxPSSH>(PsshBoxes, MP4Utilities::MakeBoxAtom('p','s','s','h'));
-			for(auto &pssh : PsshBoxes)
-			{
-				FDRMTypeFormat::FISOEncryptionInfo::FCDMInfo& cdm = iso.CDMInfos.Emplace_GetRef();
-				cdm.SystemID = pssh->GetSystemID();
-				cdm.Data = pssh->GetData();
-				for(auto &kid : pssh->GetKIDs())
-				{
-					cdm.KIDs.Emplace(kid);
-				}
-			}
-		}
-		return true;
-	};
-
-	// Get the sample type of this track. This call is necessary to trigger parsing of child nodes.
-	MP4Boxes::FMP4BoxSampleEntry::ESampleType SampleType = Stsd->GetSampleType();
-	// If already knows to not be supported, skip it.
-	if (SampleType == MP4Boxes::FMP4BoxSampleEntry::ESampleType::Unsupported)
-	{
-		OutMessage = FString::Printf(TEXT("Unsupported sample type, ignoring track #%u."), InTrackID);
-		return false;
-	}
-	// Several entries are permitted, but we need this to be unambiguous.
-	if (Stsd->GetChildren().Num() > 1)
-	{
-		OutMessage = FString::Printf(TEXT("Multiple sample descriptions found in `stsd` box, ignoring track #%u."), InTrackID);
-		return false;
-	}
-	else if (Stsd->GetChildren().Num() == 0)
-	{
-		OutMessage = FString::Printf(TEXT("No sample description found in `stsd` box, ignoring track #%u."), InTrackID);
-		return false;
-	}
-
-	OutDecoderInfo.bIsDecodable = true;
-	// Based on the sample type, get it and see if it is using codec we support.
-	if (SampleType == MP4Boxes::FMP4BoxSampleEntry::ESampleType::Video)
-	{
-		TArray<TSharedPtr<MP4Boxes::FMP4BoxVisualSampleEntry, ESPMode::ThreadSafe>> Visuals;
-		Stsd->GetSampleDescriptions(Visuals);
-		if (Visuals.Num() != 1)
-		{
-			if (Visuals.IsEmpty())
-			{
-				OutMessage = FString::Printf(TEXT("No visual sample entry found in `stsd` box, ignoring track #%u."), InTrackID);
-			}
-			else
-			{
-				OutMessage = FString::Printf(TEXT("More than one visual sample entry found in `stsd` box, ignoring track #%u."), InTrackID);
-			}
-			return false;
-		}
-
-		uint32 VisualType = Visuals[0]->GetType();
-		// Encrypted?
-		if (VisualType == MP4Utilities::MakeBoxAtom('e','n','c','v'))
-		{
-			if (!SetupEncryption(VisualType, Visuals[0]))
-			{
-				return false;
-			}
-		}
-		if (Visuals[0]->GetFrameCount() == 1)
-		{
-			OutCodecInfo.Type = Electra::FCodecTypeFormat::EType::Video;
-			OutCodecInfo.FourCC = VisualType;
-			OutCodecInfo.Properties.Emplace<Electra::FCodecTypeFormat::FVideo>();
-			Electra::FCodecTypeFormat::FVideo& cvid = OutCodecInfo.Properties.Get<Electra::FCodecTypeFormat::FVideo>();
-			cvid.Width = Visuals[0]->GetWidth();
-			cvid.Height = Visuals[0]->GetHeight();
-			auto Stts = InTrack->FindBoxRecursive<MP4Boxes::FMP4BoxSTTS>(MP4Utilities::MakeBoxAtom('s','t','t','s'), 5);
-			if (Mdhd.IsValid() && Stts.IsValid() && Stts->GetEntries().Num())
-			{
-				cvid.FrameRate = FFrameRate(Mdhd->GetTimescale(), Stts->GetEntries()[0].sample_delta);
-			}
-			for(auto& Ch : Visuals[0]->GetChildren())
-			{
-				OutCodecInfo.ExtraBoxes.Emplace(Ch->GetType(), Ch->GetBoxData());
-			}
-			if (!CodecFormatUtils->PrepareCodecTypeFormat(OutCodecInfo, IElectraCodecFormatUtilsModularFeature::FCodecTypeFormatParams()))
-			{
-				OutMessage = FString::Printf(TEXT("Failed to prepare the codec format type, ignoring track #%u."), InTrackID);
-				return false;
-			}
-			// Check if this video is decodable.
-			CheckIfDecodable();
-		}
-		else
-		{
-			OutMessage = FString::Printf(TEXT("Track #%u has a frame count other than 1 in the VisualSampleEntry, ignoring this track."), InTrackID);
-			return false;
-		}
-	}
-	else if (SampleType == MP4Boxes::FMP4BoxSampleEntry::ESampleType::Audio)
-	{
-		TArray<TSharedPtr<MP4Boxes::FMP4BoxAudioSampleEntry, ESPMode::ThreadSafe>> Audios;
-		Stsd->GetSampleDescriptions(Audios);
-		if (Audios.Num() != 1)
-		{
-			if (Audios.IsEmpty())
-			{
-				OutMessage = FString::Printf(TEXT("No audio sample entry found in `stsd` box, ignoring track #%u."), InTrackID);
-			}
-			else
-			{
-				OutMessage = FString::Printf(TEXT("More than one audio sample entry found in `stsd` box, ignoring track #%u."), InTrackID);
-			}
-			return false;
-		}
-
-		uint32 AudioType = Audios[0]->GetType();
-		// Encrypted?
-		if (AudioType == MP4Utilities::MakeBoxAtom('e','n','c','a'))
-		{
-			if (!SetupEncryption(AudioType, Audios[0]))
-			{
-				return false;
-			}
-		}
-		OutCodecInfo.Type = Electra::FCodecTypeFormat::EType::Audio;
-		OutCodecInfo.FourCC = AudioType;
-		OutCodecInfo.Properties.Emplace<Electra::FCodecTypeFormat::FAudio>();
-		Electra::FCodecTypeFormat::FAudio& caud = OutCodecInfo.Properties.Get<Electra::FCodecTypeFormat::FAudio>();
-		caud.NumChannels = Audios[0]->GetChannelCount();
-		caud.SampleRate = Audios[0]->GetSampleRate();
-		for(auto& Ch : Audios[0]->GetChildren())
-		{
-			OutCodecInfo.ExtraBoxes.Emplace(Ch->GetType(), Ch->GetBoxData());
-		}
-		if (!CodecFormatUtils->PrepareCodecTypeFormat(OutCodecInfo, IElectraCodecFormatUtilsModularFeature::FCodecTypeFormatParams()))
-		{
-			OutMessage = FString::Printf(TEXT("Failed to prepare the codec format type, ignoring track #%u."), InTrackID);
-			return false;
-		}
-		// Check if this audio is decodable.
-		CheckIfDecodable();
-	}
-	else if (SampleType == MP4Boxes::FMP4BoxSampleEntry::ESampleType::Subtitles || SampleType == MP4Boxes::FMP4BoxSampleEntry::ESampleType::TimedText)
-	{
-		TArray<TSharedPtr<MP4Boxes::FMP4BoxSampleEntry, ESPMode::ThreadSafe>> Subtitles;
-		Stsd->GetSampleDescriptions(Subtitles);
-		if (Subtitles.Num() != 1)
-		{
-			if (Subtitles.IsEmpty())
-			{
-				OutMessage = FString::Printf(TEXT("No subtitle sample entry found in `stsd` box, ignoring track #%u."), InTrackID);
-			}
-			else
-			{
-				OutMessage = FString::Printf(TEXT("More than one subtitle sample entry found in `stsd` box, ignoring track #%u."), InTrackID);
-			}
-			return false;
-		}
-
-		uint32 SubtitleType = Subtitles[0]->GetType();
-		// Encrypted?
-		if (SubtitleType == MP4Utilities::MakeBoxAtom('e','n','c','t') || SubtitleType == MP4Utilities::MakeBoxAtom('e','n','c','u'))
-		{
-			if (!SetupEncryption(SubtitleType, Subtitles[0]))
-			{
-				return false;
-			}
-		}
-		OutCodecInfo.Type = Electra::FCodecTypeFormat::EType::Subtitle;
-		OutCodecInfo.FourCC = SubtitleType;
-		OutCodecInfo.Properties.Emplace<Electra::FCodecTypeFormat::FSubtitle>();
-		//Electra::FCodecTypeFormat::FSubtitle& csub = OutCodecInfo.Properties.Get<Electra::FCodecTypeFormat::FSubtitle>();
-		// Add the sample box itself as an extra box.
-		OutCodecInfo.ExtraBoxes.Emplace(SubtitleType, Subtitles[0]->GetBoxData());
-		for(auto& Ch : Subtitles[0]->GetChildren())
-		{
-			OutCodecInfo.ExtraBoxes.Emplace(Ch->GetType(), Ch->GetBoxData());
-		}
-		if (!CodecFormatUtils->PrepareCodecTypeFormat(OutCodecInfo, IElectraCodecFormatUtilsModularFeature::FCodecTypeFormatParams()))
-		{
-			OutMessage = FString::Printf( TEXT("Failed to prepare the codec format type, ignoring track #%u."), InTrackID);
-			return false;
-		}
-	}
-	else if (SampleType == MP4Boxes::FMP4BoxSampleEntry::ESampleType::QTFFTimecode)
-	{
-		TArray<TSharedPtr<MP4Boxes::FMP4BoxQTFFTimecodeSampleEntry, ESPMode::ThreadSafe>> Timecodes;
-		Stsd->GetSampleDescriptions(Timecodes);
-		if (Timecodes.Num() != 1)
-		{
-			if (Timecodes.IsEmpty())
-			{
-				OutMessage = FString::Printf(TEXT("No timecode sample entry found in `stsd` box, ignoring track #%u."), InTrackID);
-			}
-			else
-			{
-				OutMessage = FString::Printf(TEXT("More than one timecode sample entry found in `stsd` box, ignoring track #%u."), InTrackID);
-			}
-			return false;
-		}
-
-		switch(Timecodes[0]->GetType())
-		{
-			case MP4Utilities::MakeBoxAtom('t','m','c','d'):
-			{
-				Electra::FCodecTypeFormat::FTMCDTimecode Timecode;
-
-				Timecode.Flags = Timecodes[0]->GetFlags();
-				Timecode.Timescale = Timecodes[0]->GetTimescale();
-				Timecode.FrameDuration = Timecodes[0]->GetFrameDuration();
-				Timecode.NumberOfFrames = Timecodes[0]->GetNumberOfFrames();
-				OutCodecInfo.Properties.Emplace<Electra::FCodecTypeFormat::FTMCDTimecode>(MoveTemp(Timecode));
-				OutCodecInfo.RFC6381 = MP4Utilities::GetPrintableBoxAtom(Timecodes[0]->GetType());
-				OutCodecInfo.FourCC = Timecodes[0]->GetType();
-				OutCodecInfo.DCR = Timecodes[0]->GetBoxData();
-				OutCodecInfo.Type = Electra::FCodecTypeFormat::EType::Timecode;
-				AddChildren(Timecodes[0]);
-				break;
-			}
-		}
-	}
-
-	// If the track has not been identified as usable so far, remove it.
-	if (OutCodecInfo.Type == Electra::FCodecTypeFormat::EType::Invalid)
-	{
-		TArray<TSharedPtr<MP4Boxes::FMP4BoxSampleEntry, ESPMode::ThreadSafe>> Entries;
-		Stsd->GetSampleDescriptions(Entries);
-		check(Entries.Num() == 1);
-		OutMessage = FString::Printf(TEXT("Track of sample type \"%s\" is not supported, ignoring track #%u."), *MP4Utilities::GetPrintableBoxAtom(Entries[0]->GetType()), InTrackID);
-		return false;
-	}
-	return true;
-}
-
-
 } // namespace Electra
diff --git a/Engine/Plugins/Media/ElectraPlayer/Source/ElectraPlayerRuntime/Private/Runtime/Utilities/MP4Helpers.h b/Engine/Plugins/Media/ElectraPlayer/Source/ElectraPlayerRuntime/Private/Runtime/Utilities/MP4Helpers.h
--- a/Engine/Plugins/Media/ElectraPlayer/Source/ElectraPlayerRuntime/Private/Runtime/Utilities/MP4Helpers.h
+++ b/Engine/Plugins/Media/ElectraPlayer/Source/ElectraPlayerRuntime/Private/Runtime/Utilities/MP4Helpers.h
@@ -6,12 +6,9 @@
 #include "Templates/SharedPointer.h"
 #include "HAL/FileManager.h"
 #include "Serialization/Archive.h"
-
 #include "HTTP/HTTPManager.h"
 #include "Core/MediaEventSignal.h"
 #include "MP4DataReader.h"
-#include "MP4Boxes.h"
-#include "CodecTypeFormat.h"
 
 namespace Electra
 {
@@ -64,9 +61,4 @@ namespace Electra
 		volatile bool bWasAborted = false;
 	};
 
-
-
-	bool GetTrackFormatInfo(FString& OutMessage, FCodecTypeFormat& OutCodecInfo, FDRMTypeFormat& OutDRMInfo, FDecoderInformation& OutDecoderInfo, const TSharedPtr<MP4Boxes::FMP4BoxTRAK, ESPMode::ThreadSafe>& InTrack, uint32 InTrackID, bool bInCheckIfDecodable);
-
-
 } // namespace Electra
diff --git a/Engine/Plugins/Media/ElectraPlayer/Source/ElectraProtron/Private/Player/ElectraProtronPlayerImpl.cpp b/Engine/Plugins/Media/ElectraPlayer/Source/ElectraProtron/Private/Player/ElectraProtronPlayerImpl.cpp
--- a/Engine/Plugins/Media/ElectraPlayer/Source/ElectraProtron/Private/Player/ElectraProtronPlayerImpl.cpp
+++ b/Engine/Plugins/Media/ElectraPlayer/Source/ElectraProtron/Private/Player/ElectraProtronPlayerImpl.cpp
@@ -2,18 +2,9 @@
 
 #include "ElectraProtronPlayerImpl.h"
 #include "ElectraProtronPrivate.h"
-#include "IElectraCodecFactoryModule.h"
-#include "IElectraCodecFactory.h"
-#include "IElectraCodecFormatUtils.h"
-#include "IElectraDecoderFeaturesAndOptions.h"
-#include "Utils/MPEG/ElectraUtilsMPEGVideo_H264.h"
-#include "Utils/MPEG/ElectraUtilsMPEGVideo_H265.h"
-#include "Utils/Google/ElectraUtilsVPxVideo.h"
-#include "Utils/AOMedia/ElectraUtilsAV1Video.h"
-#include "Utils/MPEG/ElectraUtilsMPEGAudio.h"
+#include "Utils/MPEG/ElectraUtilsMP4.h"
 #include "CodecTypeFormat.h"
 #include "Core/MediaThreads.h"
-#include "Features/IModularFeatures.h"
 
 /*********************************************************************************************************************/
 #include "HAL/IConsoleManager.h"
@@ -48,13 +39,6 @@ FElectraProtronPlayer::FImpl::FImpl()
 	const int32 kVideoFramesToCacheAhead = 8;
 	const int32 kVideoFramesToCacheBehind = 8;
 	CurrentSampleQueueInterface = MakeShared<FSampleQueueInterface, ESPMode::ThreadSafe>(kVideoFramesToCacheAhead, kVideoFramesToCacheBehind);
-
-	if (!CodecFormatUtils)
-	{
-		IModularFeatures::FScopedLockModularFeatureList ScopedLockModularFeatureList;
-		TArray<IElectraCodecFormatUtilsModularFeature*> FormatUtils = IModularFeatures::Get().GetModularFeatureImplementations<IElectraCodecFormatUtilsModularFeature>(IElectraCodecFormatUtilsModularFeature::GetModularFeatureName());
-		CodecFormatUtils = !FormatUtils.IsEmpty() ? FormatUtils[0] : nullptr;
-	}
 }
 
 FElectraProtronPlayer::FImpl::~FImpl()
@@ -231,12 +215,6 @@ uint32 FElectraProtronPlayer::FImpl::Run()
  */
 void FElectraProtronPlayer::FImpl::InternalOpen(const FString& InFilename)
 {
-	if (!CodecFormatUtils)
-	{
-		LastErrorMessage = TEXT("No registered IElectraCodecFormatUtilsModularFeature found. Will not be able to determine codecs.");
-		return;
-	}
-
 	// Open the file.
 	TSharedPtr<MP4Utilities::FMP4FileDataReader, ESPMode::ThreadSafe> Reader = MP4Utilities::FMP4FileDataReader::Create();
 	if (!Reader->Open(InFilename))
@@ -330,14 +308,6 @@ void FElectraProtronPlayer::FImpl::InternalOpen(const FString& InFilename)
 		}
 	}
 
-	IElectraCodecFactoryModule* FactoryModule = static_cast<IElectraCodecFactoryModule*>(FModuleManager::Get().GetModule(TEXT("ElectraCodecFactory")));
-	check(FactoryModule);
-	if (!FactoryModule)
-	{
-		LastErrorMessage = FString::Printf(TEXT("Electra decoder factory not found. Unable to use any track."));
-		return;
-	}
-
 	// Get all the tracks.
 	TArray<TSharedPtr<MP4Boxes::FMP4BoxTRAK, ESPMode::ThreadSafe>> AllTracks;
 	MoovBox->GetAllBoxInstances(AllTracks, MP4Utilities::MakeBoxAtom('t','r','a','k'));
@@ -413,7 +383,19 @@ void FElectraProtronPlayer::FImpl::InternalOpen(const FString& InFilename)
 		LongestTrackDuration = LongestTrackDuration < TrackDuration ? TrackDuration : LongestTrackDuration;
 
 		Electra::FCodecTypeFormat& CodecInfo(Tracks[i]->CodecInfo);
-		GetTrackCodecInfo(CodecInfo, ThisTrakBox, TrackID);
+		Electra::FDRMTypeFormat& DrmFormat(Tracks[i]->DrmFormat);
+		Electra::FDecoderInformation& DecoderInformation(Tracks[i]->DecoderInformation);
+		FString Msg;
+		if (!ElectraDecodersUtil::MP4::GetTrackFormatInfo(Msg, CodecInfo, DrmFormat, DecoderInformation, ThisTrakBox, TrackID, true))
+		{
+			UE_LOG(LogElectraProtron, Warning, TEXT("%s"), *Msg);
+			continue;
+		}
+		if (DrmFormat.IsEncrypted())
+		{
+			UE_LOG(LogElectraProtron, Verbose, TEXT("Encrypted tracks are not handled right now, ignoring track #%u."), TrackID);
+			continue;
+		}
 
 		// Timecode tracks are inherently usable.
 		if (CodecInfo.Type == Electra::FCodecTypeFormat::EType::Timecode)
@@ -425,70 +407,10 @@ void FElectraProtronPlayer::FImpl::InternalOpen(const FString& InFilename)
 		// Check with the decoder factory if this format can be decoded.
 		else if (CodecInfo.Type == Electra::FCodecTypeFormat::EType::Video || CodecInfo.Type == Electra::FCodecTypeFormat::EType::Audio)
 		{
-			TMap<FString, FVariant> Params;
-			switch(CodecInfo.Type)
-			{
-				case Electra::FCodecTypeFormat::EType::Video:
-				{
-					const Electra::FCodecTypeFormat::FVideo& Video(CodecInfo.Properties.Get<Electra::FCodecTypeFormat::FVideo>());
-					Params.Add(TEXT("width"), FVariant(Video.Width));
-					Params.Add(TEXT("height"), FVariant(Video.Height));
-					if (Video.FrameRate.IsValid())
-					{
-						Params.Add(TEXT("fps"), FVariant((double)Video.FrameRate.AsDecimal()));
-						Params.Add(TEXT("fps_n"), FVariant((int64)Video.FrameRate.Numerator));
-						Params.Add(TEXT("fps_d"), FVariant((uint32)Video.FrameRate.Denominator));
-					}
-					else
-					{
-						Params.Add(TEXT("fps"), FVariant((double)0.0));
-						Params.Add(TEXT("fps_n"), FVariant((int64)0));
-						Params.Add(TEXT("fps_d"), FVariant((uint32)1));
-					}
-					break;
-				}
-				case Electra::FCodecTypeFormat::EType::Audio:
-				{
-					const Electra::FCodecTypeFormat::FAudio& Audio(CodecInfo.Properties.Get<Electra::FCodecTypeFormat::FAudio>());
-					Params.Add(TEXT("channel_configuration"), FVariant(Audio.ChannelConfiguration));
-					Params.Add(TEXT("num_channels"), FVariant((int32) Audio.NumChannels));
-					Params.Add(TEXT("sample_rate"), FVariant((int32) Audio.SampleRate));
-					break;
-				}
-			}
-			Params.Add(TEXT("dcr"), FVariant(CodecInfo.DCR));
-			Params.Add(TEXT("codec_4cc"), FVariant(CodecInfo.FourCC));
-			Params.Add(TEXT("codec_name"), FVariant(CodecInfo.RFC6381));
-			// Add children box data
-			for(auto& Ch : CodecInfo.ExtraBoxes)
-			{
-				FString BoxName = FString::Printf(TEXT("$%s_box"), *MP4Utilities::GetPrintableBoxAtom(Ch.Key));
-				Params.Add(BoxName, FVariant(Ch.Value));
-			}
-
-			TMap<FString, FVariant> FormatInfo;
-			auto Factory = FactoryModule->GetBestFactoryForFormat(FormatInfo, CodecInfo.RFC6381, false, Params);
-			if (Factory.IsValid())
+			if (DecoderInformation.bIsDecodable)
 			{
 				Tracks[i]->bIsUsable = true;
-				TMap<FString, FVariant> ConfigOptions;
-				Factory->GetConfigurationOptions(ConfigOptions);
-				// Every non-video format is assumed to be keyframe-only. For video we ask the factory.
-				Tracks[i]->bIsKeyframeOnlyFormat = CodecInfo.Type != Electra::FCodecTypeFormat::EType::Video || !!ElectraDecodersUtil::GetVariantValueSafeI64(FormatInfo, IElectraDecoderFormatInfo::IsEveryFrameKeyframe, 0);
-				// See if the decoder provides a nice name for the format used.
-				Tracks[i]->HumanReadableCodecFormat = ElectraDecodersUtil::GetVariantValueFString(FormatInfo, IElectraDecoderFormatInfo::HumanReadableFormatName);
-				if (Tracks[i]->HumanReadableCodecFormat.IsEmpty())
-				{
-					Tracks[i]->HumanReadableCodecFormat = CodecInfo.HumanReadableFormatInfo;
-					if (Tracks[i]->HumanReadableCodecFormat.IsEmpty())
-					{
-						Tracks[i]->HumanReadableCodecFormat = CodecInfo.RFC6381;
-					}
-				}
-				Tracks[i]->DecoderProviderInfo.Name = Factory->GetProviderInformation().GetName();
-				Tracks[i]->DecoderProviderInfo.Version = Factory->GetProviderInformation().GetVersion();
-				Tracks[i]->DecoderProviderInfo.Implementation = Factory->GetProviderInformation().GetImplementation();
-				Tracks[i]->DecoderProviderInfo.Vendor = Factory->GetProviderInformation().GetVendor();
+				Tracks[i]->bIsKeyframeOnlyFormat = CodecInfo.KeyframeMode == Electra::FCodecTypeFormat::EKeyframeMode::OnlyKeyframes;
 				// Take note of the supported track with the longest duration.
 				LongestSupportedTrackDuration = LongestSupportedTrackDuration < TrackDuration ? TrackDuration : LongestSupportedTrackDuration;
 				// Likewise for the shortest.
@@ -501,7 +423,6 @@ void FElectraProtronPlayer::FImpl::InternalOpen(const FString& InFilename)
 		}
 		else if (CodecInfo.Type == Electra::FCodecTypeFormat::EType::Subtitle)
 		{
-			// TODO: Check with the subtitle decoder factory if this format can be decoded.
 			UE_LOG(LogElectraProtron, Verbose, TEXT("Subtitles in track #%u are not handled right now, ignoring this track."), TrackID);
 		}
 	}
@@ -747,254 +668,6 @@ void FElectraProtronPlayer::FImpl::InternalOpen(const FString& InFilename)
 	}
 }
 
-void FElectraProtronPlayer::FImpl::GetTrackCodecInfo(Electra::FCodecTypeFormat& OutCodecInfo, const TSharedPtr<MP4Boxes::FMP4BoxTRAK, ESPMode::ThreadSafe>& InTrack, uint32 InTrackID)
-{
-	// There needs to be an `stsd` box in this track. We do not try the expected path of `trak`->`mdia`->`minf`->`stbl`->`stsd` as
-	// this is not that much faster and if the file is somewhat ill-formed we may not find it if it's grouped under elsewhere.
-	auto Stsd = InTrack->FindBoxRecursive<MP4Boxes::FMP4BoxSTSD>(MP4Utilities::MakeBoxAtom('s','t','s','d'), 6);
-	if (!Stsd)
-	{
-		// If not found we just ignore the track. That's a warning but not an error.
-		UE_LOG(LogElectraProtron, Warning, TEXT("No `stsd` box found, ignoring track #%u."), InTrackID);
-		return;
-	}
-
-	// Get language code
-	auto Mdhd = InTrack->FindBoxRecursive<MP4Boxes::FMP4BoxMDHD>(MP4Utilities::MakeBoxAtom('m','d','h','d'), 2);
-	if (Mdhd)
-	{
-		Electra::BCP47::ParseRFC5646Tag(OutCodecInfo.LanguageTag, Mdhd->GetLanguageCode639_2T());
-	}
-	auto Elng = InTrack->FindBoxRecursive<MP4Boxes::FMP4BoxELNG>(MP4Utilities::MakeBoxAtom('e','l','n','g'), 6);
-	if (Elng)
-	{
-		Electra::BCP47::ParseRFC5646Tag(OutCodecInfo.LanguageTag, Elng->GetLanguageTag());
-	}
-
-	auto AddChildren = [&OutCodecInfo](const TSharedPtr<MP4Boxes::FMP4BoxBase, ESPMode::ThreadSafe>& InSample) -> void
-	{
-		for(auto& Ch : InSample->GetChildren())
-		{
-			OutCodecInfo.ExtraBoxes.Emplace(Ch->GetType(), Ch->GetBoxData());
-		}
-	};
-
-	// Get the sample type of this track. This call is necessary to trigger parsing of child nodes.
-	MP4Boxes::FMP4BoxSampleEntry::ESampleType SampleType = Stsd->GetSampleType();
-	// If already knows to not be supported, skip it.
-	if (SampleType == MP4Boxes::FMP4BoxSampleEntry::ESampleType::Unsupported)
-	{
-		UE_LOG(LogElectraProtron, Warning, TEXT("Unsupported sample type, ignoring track #%u."), InTrackID);
-		return;
-	}
-	// Several entries are permitted, but we need this to be unambiguous.
-	if (Stsd->GetChildren().Num() > 1)
-	{
-		UE_LOG(LogElectraProtron, Warning, TEXT("Multiple sample descriptions found in `stsd` box, ignoring track #%u."), InTrackID);
-		return;
-	}
-	else if (Stsd->GetChildren().Num() == 0)
-	{
-		UE_LOG(LogElectraProtron, Warning, TEXT("No sample description found in `stsd` box, ignoring track #%u."), InTrackID);
-		return;
-	}
-
-	bool bIsEncrypted = false;
-	// Based on the sample type, get it and see if it is using codec we support.
-	if (SampleType == MP4Boxes::FMP4BoxSampleEntry::ESampleType::Video)
-	{
-		TArray<TSharedPtr<MP4Boxes::FMP4BoxVisualSampleEntry, ESPMode::ThreadSafe>> Visuals;
-		Stsd->GetSampleDescriptions(Visuals);
-		if (Visuals.Num() != 1)
-		{
-			if (Visuals.IsEmpty())
-			{
-				UE_LOG(LogElectraProtron, Warning, TEXT("No visual sample entry found in `stsd` box, ignoring track #%u."), InTrackID);
-			}
-			else
-			{
-				UE_LOG(LogElectraProtron, Warning, TEXT("More than one visual sample entry found in `stsd` box, ignoring track #%u."), InTrackID);
-			}
-			return;
-		}
-
-		switch(Visuals[0]->GetType())
-		{
-			case MP4Utilities::MakeBoxAtom('e','n','c','v'):
-			{
-				bIsEncrypted = true;
-				break;
-			}
-			default:
-			{
-				if (Visuals[0]->GetFrameCount() == 1)
-				{
-					OutCodecInfo.Type = Electra::FCodecTypeFormat::EType::Video;
-					OutCodecInfo.FourCC = Visuals[0]->GetType();
-					OutCodecInfo.Properties.Emplace<Electra::FCodecTypeFormat::FVideo>();
-					Electra::FCodecTypeFormat::FVideo& cvid = OutCodecInfo.Properties.Get<Electra::FCodecTypeFormat::FVideo>();
-					cvid.Width = Visuals[0]->GetWidth();
-					cvid.Height = Visuals[0]->GetHeight();
-					auto Stts = InTrack->FindBoxRecursive<MP4Boxes::FMP4BoxSTTS>(MP4Utilities::MakeBoxAtom('s','t','t','s'), 5);
-					if (Mdhd.IsValid() && Stts.IsValid() && Stts->GetEntries().Num())
-					{
-						cvid.FrameRate = FFrameRate(Mdhd->GetTimescale(), Stts->GetEntries()[0].sample_delta);
-					}
-					for(auto& Ch : Visuals[0]->GetChildren())
-					{
-						OutCodecInfo.ExtraBoxes.Emplace(Ch->GetType(), Ch->GetBoxData());
-					}
-					if (!CodecFormatUtils->PrepareCodecTypeFormat(OutCodecInfo, IElectraCodecFormatUtilsModularFeature::FCodecTypeFormatParams()))
-					{
-						UE_LOG(LogElectraProtron, Warning, TEXT("Failed to prepare the codec format type, ignoring track #%u."), InTrackID);
-						return;
-					}
-				}
-				else
-				{
-					UE_LOG(LogElectraProtron, Warning, TEXT("Track #%u has a frame count other than 1 in the VisualSampleEntry, ignoring this track."), InTrackID);
-				}
-				break;
-			}
-		}
-	}
-	else if (SampleType == MP4Boxes::FMP4BoxSampleEntry::ESampleType::Audio)
-	{
-		TArray<TSharedPtr<MP4Boxes::FMP4BoxAudioSampleEntry, ESPMode::ThreadSafe>> Audios;
-		Stsd->GetSampleDescriptions(Audios);
-		if (Audios.Num() != 1)
-		{
-			if (Audios.IsEmpty())
-			{
-				UE_LOG(LogElectraProtron, Warning, TEXT("No audio sample entry found in `stsd` box, ignoring track #%u."), InTrackID);
-			}
-			else
-			{
-				UE_LOG(LogElectraProtron, Warning, TEXT("More than one audio sample entry found in `stsd` box, ignoring track #%u."), InTrackID);
-			}
-			return;
-		}
-
-		switch(Audios[0]->GetType())
-		{
-			case MP4Utilities::MakeBoxAtom('e','n','c','a'):
-			{
-				bIsEncrypted = true;
-				break;
-			}
-			default:
-			{
-				OutCodecInfo.Type = Electra::FCodecTypeFormat::EType::Audio;
-				OutCodecInfo.FourCC = Audios[0]->GetType();
-				OutCodecInfo.Properties.Emplace<Electra::FCodecTypeFormat::FAudio>();
-				Electra::FCodecTypeFormat::FAudio& caud = OutCodecInfo.Properties.Get<Electra::FCodecTypeFormat::FAudio>();
-				caud.NumChannels = Audios[0]->GetChannelCount();
-				caud.SampleRate = Audios[0]->GetSampleRate();
-				for(auto& Ch : Audios[0]->GetChildren())
-				{
-					OutCodecInfo.ExtraBoxes.Emplace(Ch->GetType(), Ch->GetBoxData());
-				}
-				if (!CodecFormatUtils->PrepareCodecTypeFormat(OutCodecInfo, IElectraCodecFormatUtilsModularFeature::FCodecTypeFormatParams()))
-				{
-					UE_LOG(LogElectraProtron, Warning, TEXT("Failed to prepare the codec format type, ignoring track #%u."), InTrackID);
-					return;
-				}
-				break;
-			}
-		}
-
-	}
-	else if (SampleType == MP4Boxes::FMP4BoxSampleEntry::ESampleType::Subtitles || SampleType == MP4Boxes::FMP4BoxSampleEntry::ESampleType::TimedText)
-	{
-		TArray<TSharedPtr<MP4Boxes::FMP4BoxSampleEntry, ESPMode::ThreadSafe>> Subtitles;
-		Stsd->GetSampleDescriptions(Subtitles);
-		if (Subtitles.Num() != 1)
-		{
-			if (Subtitles.IsEmpty())
-			{
-				UE_LOG(LogElectraProtron, Warning, TEXT("No subtitle sample entry found in `stsd` box, ignoring track #%u."), InTrackID);
-			}
-			else
-			{
-				UE_LOG(LogElectraProtron, Warning, TEXT("More than one subtitle sample entry found in `stsd` box, ignoring track #%u."), InTrackID);
-			}
-			return;
-		}
-
-		switch(Subtitles[0]->GetType())
-		{
-			case MP4Utilities::MakeBoxAtom('e','n','c','t'):
-			case MP4Utilities::MakeBoxAtom('e','n','c','u'):
-			{
-				bIsEncrypted = true;
-				break;
-			}
-			default:
-			{
-				OutCodecInfo.Type = Electra::FCodecTypeFormat::EType::Subtitle;
-				OutCodecInfo.FourCC = Subtitles[0]->GetType();
-				OutCodecInfo.Properties.Emplace<Electra::FCodecTypeFormat::FSubtitle>();
-				Electra::FCodecTypeFormat::FSubtitle& csub = OutCodecInfo.Properties.Get<Electra::FCodecTypeFormat::FSubtitle>();
-				for(auto& Ch : Subtitles[0]->GetChildren())
-				{
-					OutCodecInfo.ExtraBoxes.Emplace(Ch->GetType(), Ch->GetBoxData());
-				}
-			/*
-				if (!CodecFormatUtils->PrepareCodecTypeFormat(OutCodecInfo, IElectraCodecFormatUtilsModularFeature::FCodecTypeFormatParams()))
-				{
-					UE_LOG(LogElectraProtron, Warning, TEXT("Failed to prepare the codec format type, ignoring track #%u."), InTrackID);
-					return;
-				}
-			*/
-				break;
-			}
-		}
-	}
-	else if (SampleType == MP4Boxes::FMP4BoxSampleEntry::ESampleType::QTFFTimecode)
-	{
-		TArray<TSharedPtr<MP4Boxes::FMP4BoxQTFFTimecodeSampleEntry, ESPMode::ThreadSafe>> Timecodes;
-		Stsd->GetSampleDescriptions(Timecodes);
-		check(Timecodes.Num() == 1);
-		switch(Timecodes[0]->GetType())
-		{
-			case MP4Utilities::MakeBoxAtom('t','m','c','d'):
-			{
-				Electra::FCodecTypeFormat::FTMCDTimecode Timecode;
-
-				Timecode.Flags = Timecodes[0]->GetFlags();
-				Timecode.Timescale = Timecodes[0]->GetTimescale();
-				Timecode.FrameDuration = Timecodes[0]->GetFrameDuration();
-				Timecode.NumberOfFrames = Timecodes[0]->GetNumberOfFrames();
-				OutCodecInfo.Properties.Emplace<Electra::FCodecTypeFormat::FTMCDTimecode>(MoveTemp(Timecode));
-				OutCodecInfo.RFC6381 = MP4Utilities::GetPrintableBoxAtom(Timecodes[0]->GetType());
-				OutCodecInfo.FourCC = Timecodes[0]->GetType();
-				OutCodecInfo.DCR = Timecodes[0]->GetBoxData();
-				OutCodecInfo.Type = Electra::FCodecTypeFormat::EType::Timecode;
-				AddChildren(Timecodes[0]);
-				break;
-			}
-		}
-	}
-
-	// If the track has not been identified as usable so far, remove it.
-	if (OutCodecInfo.Type == Electra::FCodecTypeFormat::EType::Invalid)
-	{
-		if (!bIsEncrypted)
-		{
-			TArray<TSharedPtr<MP4Boxes::FMP4BoxSampleEntry, ESPMode::ThreadSafe>> Entries;
-			Stsd->GetSampleDescriptions(Entries);
-			check(Entries.Num() == 1);
-			UE_LOG(LogElectraProtron, Warning, TEXT("Track of sample type \"%s\" is not supported, ignoring track #%u."), *MP4Utilities::GetPrintableBoxAtom(Entries[0]->GetType()), InTrackID);
-		}
-		else
-		{
-			UE_LOG(LogElectraProtron, Warning, TEXT("Track is using encryption, ignoring track #%u."), InTrackID);
-		}
-	}
-}
-
-
-
 FVariant FElectraProtronPlayer::FImpl::GetMediaInfo(FName InInfoName)
 {
 	if (InInfoName == ElectraProtronOptionNames::StartTimecodeValue || InInfoName == ElectraProtronOptionNames::StartTimecodeFrameRate ||
@@ -1018,9 +691,9 @@ FVariant FElectraProtronPlayer::FImpl::GetMediaInfo(FName InInfoName)
 			{
 				return FVariant(ti.bIsKeyframeOnlyFormat ? (int32)1 : (int32)0);
 			}
-			else if (InInfoName == ElectraProtronOptionNames::VideoDecoderProvider && ti.DecoderProviderInfo.Name.Len())
+			else if (InInfoName == ElectraProtronOptionNames::VideoDecoderProvider && ti.DecoderInformation.ProviderInfo.Name.Len())
 			{
-				return FVariant(ti.DecoderProviderInfo.Name);
+				return FVariant(ti.DecoderInformation.ProviderInfo.Name);
 			}
 		}
 	}
@@ -1031,9 +704,9 @@ FVariant FElectraProtronPlayer::FImpl::GetMediaInfo(FName InInfoName)
 		if (SelectedAudioTrackIndex >= 0 && SelectedAudioTrackIndex < UsableTrackArrayIndicesByType[ci].Num())
 		{
 			const FTrackInfo& ti(*Tracks[UsableTrackArrayIndicesByType[ci][SelectedAudioTrackIndex]]);
-			if (InInfoName == ElectraProtronOptionNames::AudioDecoderProvider && ti.DecoderProviderInfo.Name.Len())
+			if (InInfoName == ElectraProtronOptionNames::AudioDecoderProvider && ti.DecoderInformation.ProviderInfo.Name.Len())
 			{
-				return FVariant(ti.DecoderProviderInfo.Name);
+				return FVariant(ti.DecoderInformation.ProviderInfo.Name);
 			}
 		}
 	}
@@ -1208,7 +881,7 @@ bool FElectraProtronPlayer::FImpl::GetVideoTrackFormat(int32 InTrackIndex, int32
 			OutFormat.FrameRate = 0.0f;
 			OutFormat.FrameRates = TRange<float>{ OutFormat.FrameRate };//TRange<float>::Empty();
 		}
-		OutFormat.TypeName = ti.HumanReadableCodecFormat;
+		OutFormat.TypeName = ti.CodecInfo.HumanReadableFormatInfo;
 		return true;
 	}
 	return false;
@@ -1223,7 +896,7 @@ bool FElectraProtronPlayer::FImpl::GetAudioTrackFormat(int32 InTrackIndex, int32
 		OutFormat.BitsPerSample = 16;
 		OutFormat.NumChannels = ai.NumChannels;
 		OutFormat.SampleRate = ai.SampleRate;
-		OutFormat.TypeName = ti.HumanReadableCodecFormat;
+		OutFormat.TypeName = ti.CodecInfo.HumanReadableFormatInfo;
 		return true;
 	}
 	return false;
diff --git a/Engine/Plugins/Media/ElectraPlayer/Source/ElectraProtron/Private/Player/ElectraProtronPlayerImpl.h b/Engine/Plugins/Media/ElectraPlayer/Source/ElectraProtron/Private/Player/ElectraProtronPlayerImpl.h
--- a/Engine/Plugins/Media/ElectraPlayer/Source/ElectraProtron/Private/Player/ElectraProtronPlayerImpl.h
+++ b/Engine/Plugins/Media/ElectraPlayer/Source/ElectraProtron/Private/Player/ElectraProtronPlayerImpl.h
@@ -16,12 +16,9 @@
 #include "MP4Boxes.h"
 #include "MP4Track.h"
 #include "CodecTypeFormat.h"
-#include "IElectraDecoder.h"
-#include "IElectraCodecFormatUtils.h"
 #include "ElectraDecodersPlatformResources.h"
 #include "MediaSamples.h"
 #include "ElectraTextureSample.h"
-#include "Utils/MPEG/ElectraUtilsMPEGVideo.h"
 #include "Utils/AudioChannelMapper.h"
 #include "Utils/VideoDecoderHelpers.h"
 #include "ElectraProtronPlayerCache.h"
@@ -344,16 +341,10 @@ private:
 	{
 		TArray<TWeakPtr<FTrackInfo, ESPMode::ThreadSafe>> IsReferencedByTracks;
 		Electra::FCodecTypeFormat CodecInfo;
-		FString HumanReadableCodecFormat;
-		struct FDecoderProviderInfo
-		{
-			FString Name;
-			FString Version;
-			FString Implementation;
-			FString Vendor;
-		};
+		Electra::FDRMTypeFormat DrmFormat;
+		Electra::FDecoderInformation DecoderInformation;
+
 		FTimespan Duration;
-		FDecoderProviderInfo DecoderProviderInfo;
 		TSharedPtr<MP4Boxes::FMP4BoxTRAK, ESPMode::ThreadSafe> TrackBox;
 		TSharedPtr<MP4Boxes::FMP4Track::FFragmentInfo, ESPMode::ThreadSafe> FragmentInfo;
 		TSharedPtr<MP4Boxes::FMP4Track, ESPMode::ThreadSafe> MP4Track;
@@ -417,7 +408,6 @@ private:
 	void SendWorkerThreadMessage(FWorkerThreadMessage&& InMessage);
 
 	void InternalOpen(const FString& InFilename);
-	void GetTrackCodecInfo(Electra::FCodecTypeFormat& OutCodecInfo, const TSharedPtr<MP4Boxes::FMP4BoxTRAK, ESPMode::ThreadSafe>& InTrack, uint32 InTrackID);
 
 	void UpdateTrackLoader(int32 InCodecTypeIndex);
 
@@ -435,8 +425,6 @@ private:
 
 	FConfig Config;
 
-	IElectraCodecFormatUtilsModularFeature* CodecFormatUtils = nullptr;
-
 	FRunnableThread* Thread = nullptr;
 	ImplPointer SelfDuringTerminate;
 	FMediaEvent WorkMessageSignal;
