ElectraPlayer: switching back to previous stream segment reader while investigating a bug with the new one

#rnx

#Synced-CL 51333923

diff --git a/Engine/Plugins/Media/ElectraPlayer/Source/ElectraPlayerRuntime/Private/Runtime/Player/StreamSegmentReaderCommon.cpp b/Engine/Plugins/Media/ElectraPlayer/Source/ElectraPlayerRuntime/Private/Runtime/Player/StreamSegmentReaderCommon.cpp
--- a/Engine/Plugins/Media/ElectraPlayer/Source/ElectraPlayerRuntime/Private/Runtime/Player/StreamSegmentReaderCommon.cpp
+++ b/Engine/Plugins/Media/ElectraPlayer/Source/ElectraPlayerRuntime/Private/Runtime/Player/StreamSegmentReaderCommon.cpp
@@ -46,7 +46,7 @@ DECLARE_CYCLE_STAT(TEXT("FStreamSegmentReaderCommon::HandleRequest"), STAT_Elect
 namespace Electra { namespace Stream {
 TAutoConsoleVariable<bool> CVarElectraStreamUseLegacyMP4(
 	TEXT("Electra.UseLegacyMP4Stream"),
-	false,
+	true,
 	TEXT("Use old mp4 parser for DASH and HLS streaming.\n")
 	TEXT(" true: use old parser, false use new parser"),
 	ECVF_Default);
