ElectraPlayer: it is now possible to change video tracks of an mp4 that have different resolutions

#rnx

#Synced-CL 51366553

diff --git a/Engine/Plugins/Media/ElectraPlayer/Source/ElectraPlayerRuntime/Private/Runtime/Player/AdaptiveStreamingPlayer.cpp b/Engine/Plugins/Media/ElectraPlayer/Source/ElectraPlayerRuntime/Private/Runtime/Player/AdaptiveStreamingPlayer.cpp
--- a/Engine/Plugins/Media/ElectraPlayer/Source/ElectraPlayerRuntime/Private/Runtime/Player/AdaptiveStreamingPlayer.cpp
+++ b/Engine/Plugins/Media/ElectraPlayer/Source/ElectraPlayerRuntime/Private/Runtime/Player/AdaptiveStreamingPlayer.cpp
@@ -3751,6 +3751,7 @@ void FAdaptiveStreamingPlayer::InternalHandleSegmentTrackChanges(const FTimeValu
 						{
 							StreamRcvBuf->SelectTrackWhenAvailable(CurrentPlaybackSequenceID[StreamTypeToArrayIndex(InType)], BufferSourceInfo);
 						}
+						UpdatePeriodStreamBufferSourceInfo(InCurrentPeriod, InType, BufferSourceInfo);
 						OutSelectionAttributes = *InPendingSelection;
 						bChangeMade = true;
 					}
diff --git a/Engine/Plugins/Media/ElectraPlayer/Source/ElectraPlayerRuntime/Private/Runtime/Player/AdaptiveStreamingPlayerDecoders.cpp b/Engine/Plugins/Media/ElectraPlayer/Source/ElectraPlayerRuntime/Private/Runtime/Player/AdaptiveStreamingPlayerDecoders.cpp
--- a/Engine/Plugins/Media/ElectraPlayer/Source/ElectraPlayerRuntime/Private/Runtime/Player/AdaptiveStreamingPlayerDecoders.cpp
+++ b/Engine/Plugins/Media/ElectraPlayer/Source/ElectraPlayerRuntime/Private/Runtime/Player/AdaptiveStreamingPlayerDecoders.cpp
@@ -124,15 +124,12 @@ bool FAdaptiveStreamingPlayer::CanDecodeSubtitle(const FString& MimeType, const
 
 
 
-bool FAdaptiveStreamingPlayer::FindMatchingStreamInfo(FStreamCodecInformation& OutStreamInfo, const FString& InPeriodID, const FTimeValue& AtTime, const FStreamCodecInformation& InForCodec)
+bool FAdaptiveStreamingPlayer::GetHighestVideoProfileResolution(FStreamCodecInformation& OutStreamInfo, TSharedPtrTS<const FBufferSourceInfo> InBufferSource, const FTimeValue& AtTime, const FStreamCodecInformation& InForCodec)
 {
-	TSharedPtrTS<ITimelineMediaAsset> Asset;
-
-	// Not used at the moment.
-	int32 MaxWidth = 0;
-	int32 MaxHeight = 0;
-	MaxWidth = MaxWidth ? MaxWidth : 32768;
-	MaxHeight = MaxHeight ? MaxHeight : 32768;
+	if (!InBufferSource)
+	{
+		return false;
+	}
 
 	// Locate the period by ID or the specified time.
 	ActivePeriodCriticalSection.Lock();
@@ -141,88 +138,73 @@ bool FAdaptiveStreamingPlayer::FindMatchingStreamInfo(FStreamCodecInformation& O
 		ActivePeriodCriticalSection.Unlock();
 		return false;
 	}
-	bool bFound = false;
-	if (InPeriodID.Len())
+	TSharedPtrTS<ITimelineMediaAsset> Asset;
+	const FString& PeriodID(InBufferSource->PeriodID);
+	if (PeriodID.Len())
 	{
 		for(int32 i=0; i<ActivePeriods.Num(); ++i)
 		{
-			if (ActivePeriods[i].ID.Equals(InPeriodID))
+			if (ActivePeriods[i].ID.Equals(PeriodID))
 			{
 				Asset = ActivePeriods[i].Period;
-				bFound = true;
 				break;
 			}
 		}
 	}
-	if (!bFound && AtTime.IsValid())
+	if (!Asset && AtTime.IsValid())
 	{
 		for(int32 i=0; i<ActivePeriods.Num(); ++i)
 		{
 			if (AtTime >= ActivePeriods[i].TimeRange.Start && AtTime < (ActivePeriods[i].TimeRange.End.IsValid() ? ActivePeriods[i].TimeRange.End : FTimeValue::GetPositiveInfinity()))
 			{
 				Asset = ActivePeriods[i].Period;
-				bFound = true;
 			}
 		}
 	}
-	if (!bFound)
+	if (!Asset)
 	{
 		Asset = ActivePeriods[0].Period;
 	}
 	ActivePeriodCriticalSection.Unlock();
 	if (Asset.IsValid())
 	{
-		for(int32 nAdapt=0, nAdaptMax=Asset->GetNumberOfAdaptationSets(EStreamType::Video); nAdapt<nAdaptMax; ++nAdapt)
+		TSharedPtrTS<IPlaybackAssetAdaptationSet> VideoAdaptationSet;
+		if (InBufferSource->HardIndex >= 0 && InBufferSource->HardIndex < Asset->GetNumberOfAdaptationSets(EStreamType::Video))
 		{
-			TSharedPtrTS<IPlaybackAssetAdaptationSet> VideoSet = Asset->GetAdaptationSetByTypeAndIndex(EStreamType::Video, nAdapt);
-			if (VideoSet.IsValid())
+			VideoAdaptationSet = Asset->GetAdaptationSetByTypeAndIndex(EStreamType::Video, InBufferSource->HardIndex);
+		}
+		if (VideoAdaptationSet && VideoAdaptationSet->GetNumberOfRepresentations())
+		{
+			TArray<FStreamCodecInformation> VideoCodecInfos;
+			for(int32 i=0, iMax=VideoAdaptationSet->GetNumberOfRepresentations(); i<iMax; ++i)
 			{
-				TArray<FStreamCodecInformation> VideoCodecInfos;
-
-				// Filter the streams by permitted max.resolution
-				int32 LowestBW = TNumericLimits<int32>::Max();
-				for(int32 i=0, iMax=VideoSet->GetNumberOfRepresentations(); i<iMax; ++i)
+				VideoCodecInfos.Push(VideoAdaptationSet->GetRepresentationByIndex(i)->GetCodecInformation());
+			}
+			// The representations could be a mix of profile and resolutions, although not very likely.
+			// We consider profile, level and resolution separately, which may produce an amalgamation of values that
+			// does not appear as such on any particular stream. We do need to know the worst profile and resolution
+			// that can occur.
+			OutStreamInfo = VideoCodecInfos[0];
+			int32 Profile = VideoCodecInfos[0].GetProfile();
+			int32 Level = VideoCodecInfos[0].GetProfileLevel();
+			for(int32 i=1; i<VideoCodecInfos.Num(); ++i)
+			{
+				if (VideoCodecInfos[i].GetResolution().Width * VideoCodecInfos[i].GetResolution().Height > OutStreamInfo.GetResolution().Width * OutStreamInfo.GetResolution().Height)
 				{
-					if (VideoSet->GetRepresentationByIndex(i)->GetBitrate() < LowestBW)
-					{
-						LowestBW = VideoSet->GetRepresentationByIndex(i)->GetBitrate();
-					}
-					const FStreamCodecInformation& ci = VideoSet->GetRepresentationByIndex(i)->GetCodecInformation();
-					if (!ci.GetResolution().ExceedsLimit(MaxWidth, MaxHeight))
-					{
-						VideoCodecInfos.Push(ci);
-					}
+					OutStreamInfo = VideoCodecInfos[i];
 				}
-				if (VideoCodecInfos.Num() == 0)
+				if (VideoCodecInfos[i].GetProfile() > Profile)
 				{
-					for(int32 i=0, iMax=VideoSet->GetNumberOfRepresentations(); i<iMax; ++i)
-					{
-						if (VideoSet->GetRepresentationByIndex(i)->GetBitrate() <= LowestBW)
-						{
-							VideoCodecInfos.Push(VideoSet->GetRepresentationByIndex(i)->GetCodecInformation());
-							break;
-						}
-					}
+					Profile = VideoCodecInfos[i].GetProfile();
 				}
-
-				// Check codec.
-				if (VideoCodecInfos[0].GetCodec() != InForCodec.GetCodec())
+				if (VideoCodecInfos[i].GetProfileLevel() > Level)
 				{
-					continue;
+					Level = VideoCodecInfos[i].GetProfileLevel();
 				}
-
-				// Reduce to those having the highest profile.
-				VideoCodecInfos.Sort([](const FStreamCodecInformation& a, const FStreamCodecInformation& b) { return a.GetProfile() > b.GetProfile(); });
-				VideoCodecInfos = VideoCodecInfos.FilterByPredicate([BestProfile=VideoCodecInfos[0].GetProfile()](const FStreamCodecInformation& e)	{ return e.GetProfile() >= BestProfile; });
-				// Reduce to those having the highest level.
-				VideoCodecInfos.Sort([](const FStreamCodecInformation& a, const FStreamCodecInformation& b) { return a.GetProfileLevel() > b.GetProfileLevel(); });
-				VideoCodecInfos = VideoCodecInfos.FilterByPredicate([BestLevel=VideoCodecInfos[0].GetProfileLevel()](const FStreamCodecInformation& e) { return e.GetProfileLevel() >= BestLevel; });
-				// Sort by resolution
-				VideoCodecInfos.Sort([](const FStreamCodecInformation& a, const FStreamCodecInformation& b) { return a.GetResolution().Width*a.GetResolution().Height > b.GetResolution().Width*b.GetResolution().Height; });
-
-				OutStreamInfo = VideoCodecInfos[0];
-				return true;
 			}
+			OutStreamInfo.SetProfile(Profile);
+			OutStreamInfo.SetProfileLevel(Level);
+			return true;
 		}
 	}
 	return false;
@@ -251,7 +233,6 @@ int32 FAdaptiveStreamingPlayer::CreateDecoder(EStreamType type)
 			{
 				FTimeValue DecodeTime = AccessUnit->EarliestPTS.IsValid() ? AccessUnit->EarliestPTS : AccessUnit->PTS;
 				check(AccessUnit->BufferSourceInfo.IsValid());
-				FString PeriodID = AccessUnit->BufferSourceInfo.IsValid() ? AccessUnit->BufferSourceInfo->PeriodID : FString();
 				VideoDecoder.CurrentCodecInfo.Clear();
 				VideoDecoder.LastBufferSourceInfo = AccessUnit->BufferSourceInfo;
 				if (AccessUnit->AUCodecData.IsValid())
@@ -261,13 +242,14 @@ int32 FAdaptiveStreamingPlayer::CreateDecoder(EStreamType type)
 					GetMediaInfoDictionary().Set(MediaInfoVideoBitrateCurrent.Resolve(), FVariantValue((int64)VideoDecoder.CurrentCodecInfo.GetBitrate()));
 					DispatchEvent(FMetricEvent::ReportCodecFormatChange(VideoDecoder.CurrentCodecInfo));
 				}
-				FAccessUnit::Release(AccessUnit);
 
 				// Get the largest stream resolution of the currently selected video adaptation set.
 				// This is only an initial selection as there could be other adaptation sets in upcoming periods
 				// that have a larger resolution that is still within the allowed limits.
 				FStreamCodecInformation HighestStream;
-				if (!FindMatchingStreamInfo(HighestStream, PeriodID, DecodeTime, VideoDecoder.CurrentCodecInfo))
+				bOk = GetHighestVideoProfileResolution(HighestStream, AccessUnit->BufferSourceInfo, DecodeTime, VideoDecoder.CurrentCodecInfo);
+				FAccessUnit::Release(AccessUnit);
+				if (!bOk)
 				{
 					FErrorDetail err;
 					err.SetFacility(Facility::EFacility::Player);
@@ -690,7 +672,7 @@ void FAdaptiveStreamingPlayer::FeedDecoder(EStreamType Type, IAccessUnitBufferIn
 				{
 					// Get the new highest stream properties and check if the decoder can still handle those.
 					FStreamCodecInformation HighestStream;
-					if (FindMatchingStreamInfo(HighestStream, PeekedAU->BufferSourceInfo->PeriodID, PeekedAU->EarliestPTS.IsValid() ? PeekedAU->EarliestPTS : PeekedAU->PTS, PeekedAU->AUCodecData->ParsedInfo))
+					if (GetHighestVideoProfileResolution(HighestStream, PeekedAU->BufferSourceInfo, PeekedAU->EarliestPTS.IsValid() ? PeekedAU->EarliestPTS : PeekedAU->PTS, PeekedAU->AUCodecData->ParsedInfo))
 					{
 						if (!VideoDecoder.Decoder->Reopen(PeekedAU->AUCodecData, PlayerOptions.GetDictionary(), &HighestStream))
 						{
diff --git a/Engine/Plugins/Media/ElectraPlayer/Source/ElectraPlayerRuntime/Private/Runtime/Player/AdaptiveStreamingPlayerInternal.h b/Engine/Plugins/Media/ElectraPlayer/Source/ElectraPlayerRuntime/Private/Runtime/Player/AdaptiveStreamingPlayerInternal.h
--- a/Engine/Plugins/Media/ElectraPlayer/Source/ElectraPlayerRuntime/Private/Runtime/Player/AdaptiveStreamingPlayerInternal.h
+++ b/Engine/Plugins/Media/ElectraPlayer/Source/ElectraPlayerRuntime/Private/Runtime/Player/AdaptiveStreamingPlayerInternal.h
@@ -2169,7 +2169,7 @@ private:
 
 	int32 CreateDecoder(EStreamType type);
 	void DestroyDecoders();
-	bool FindMatchingStreamInfo(FStreamCodecInformation& OutStreamInfo, const FString& InPeriodID, const FTimeValue& AtTime, const FStreamCodecInformation& InForCodec);
+	bool GetHighestVideoProfileResolution(FStreamCodecInformation& OutStreamInfo, TSharedPtrTS<const FBufferSourceInfo> InBufferSource, const FTimeValue& AtTime, const FStreamCodecInformation& InForCodec);
 	void UpdateStreamResolutionLimit();
 	void AddUpcomingPeriod(TSharedPtrTS<IManifest::IPlayPeriod> InUpcomingPeriod);
 	void RequestNewPeriodStreams(EStreamType InType, FPendingSegmentRequest& InOutCurrentRequest);
