Protron: Await one sample from each selected track (video and audio, if enabled) to be ready before sending them into the sample queues. This will avoid unsynchronized consumption of one or the other (mostly consuming audio earlier than video)

#rnx

#Synced-CL 51368660

diff --git a/Engine/Plugins/Media/ElectraPlayer/Source/ElectraProtron/Private/Player/ElectraProtronPlayerControls.cpp b/Engine/Plugins/Media/ElectraPlayer/Source/ElectraProtron/Private/Player/ElectraProtronPlayerControls.cpp
--- a/Engine/Plugins/Media/ElectraPlayer/Source/ElectraProtron/Private/Player/ElectraProtronPlayerControls.cpp
+++ b/Engine/Plugins/Media/ElectraPlayer/Source/ElectraProtron/Private/Player/ElectraProtronPlayerControls.cpp
@@ -62,6 +62,14 @@ void FElectraProtronPlayer::FImpl::HandleActiveTrackChanges()
 						auto tsb = TrackSampleBuffers[Track->TrackID];
 						if (tsb.IsValid())
 						{
+							if (InCodecType == CodecTypeIndex(Electra::FCodecTypeFormat::EType::Video))
+							{
+								CurrentSampleQueueInterface->SetVideoTrackEnabled(true);
+							}
+							else if (InCodecType == CodecTypeIndex(Electra::FCodecTypeFormat::EType::Audio))
+							{
+								CurrentSampleQueueInterface->SetAudioTrackEnabled(true);
+							}
 							InDecoderThread.SetSampleBuffer(tsb, FGetSampleDlg::CreateRaw(&InLoaderThread, &FLoaderThread::GetSample));
 							TrackSelection.ActiveTrackIndex[InCodecType] = TrackSelection.SelectedTrackIndex[InCodecType];
 						}
@@ -70,6 +78,14 @@ void FElectraProtronPlayer::FImpl::HandleActiveTrackChanges()
 				else
 				{
 					// Deselect.
+					if (InCodecType == CodecTypeIndex(Electra::FCodecTypeFormat::EType::Video))
+					{
+						CurrentSampleQueueInterface->SetVideoTrackEnabled(false);
+					}
+					else if (InCodecType == CodecTypeIndex(Electra::FCodecTypeFormat::EType::Audio))
+					{
+						CurrentSampleQueueInterface->SetAudioTrackEnabled(false);
+					}
 					InDecoderThread.DisconnectSampleBuffer();
 					TrackSelection.ActiveTrackIndex[InCodecType] = TrackSelection.SelectedTrackIndex[InCodecType];
 				}
diff --git a/Engine/Plugins/Media/ElectraPlayer/Source/ElectraProtron/Private/Player/ElectraProtronPlayerDecoder.cpp b/Engine/Plugins/Media/ElectraPlayer/Source/ElectraProtron/Private/Player/ElectraProtronPlayerDecoder.cpp
--- a/Engine/Plugins/Media/ElectraPlayer/Source/ElectraProtron/Private/Player/ElectraProtronPlayerDecoder.cpp
+++ b/Engine/Plugins/Media/ElectraPlayer/Source/ElectraProtron/Private/Player/ElectraProtronPlayerDecoder.cpp
@@ -57,6 +57,9 @@ void FElectraProtronPlayer::FImpl::FDecoderThread::StopThread()
 
 uint32 FElectraProtronPlayer::FImpl::FDecoderThread::Run()
 {
+	double PreviousTime = FPlatformTime::Seconds();
+	const double kDefaultMinLoopSleepTime = 0.005;
+	MinLoopSleepTime = kDefaultMinLoopSleepTime;
 	while(!bTerminate)
 	{
 		// When handling to the target seek time we want to go as quickly as possible,
@@ -67,7 +70,14 @@ uint32 FElectraProtronPlayer::FImpl::FDecoderThread::Run()
 		}
 		else
 		{
-			WorkSignal.WaitTimeoutAndReset(1000 * 5);
+			double TimeNow = FPlatformTime::Seconds();
+			double ElapsedTime = TimeNow - PreviousTime;
+			if (ElapsedTime < MinLoopSleepTime)
+			{
+				FPlatformProcess::Sleep(MinLoopSleepTime - ElapsedTime);
+			}
+			PreviousTime = TimeNow;
+			MinLoopSleepTime = kDefaultMinLoopSleepTime;
 		}
 		if (bFlushPending)
 		{
@@ -535,12 +545,7 @@ void FElectraProtronPlayer::FImpl::FDecoderThread::DecodeOneFrame()
 	{
 		return;
 	}
-	// If we are still hanging on to output we could not deliver we will not
-	// decode new data that will likewise not be able to deliver.
-	if (CurrentDecoderOutput.IsValid())
-	{
-		return;
-	}
+
 	// If the end of the stream has been reached we are done.
 	if (bReachedEnd)
 	{
@@ -845,7 +850,11 @@ void FElectraProtronPlayer::FImpl::FDecoderThread::HandleOutputFrame()
 
 			// When sending the output we need to wait until there is room to receive it.
 			// If the output is not to be sent we can handle it quickly without waiting.
-			if (!bSendToOutput || Params.SampleQueueInterface->CanEnqueueVideoSample(RawPTS))
+			if (bSendToOutput)
+			{
+				Params.SampleQueueInterface->VideoSampleWillBeReady();
+			}
+			if (!bSendToOutput || (Params.SampleQueueInterface->CanEnqueueVideoSample(RawPTS) && Params.SampleQueueInterface->AreRequiredSamplesReady()))
 			{
 				// First perform any processing that needs to be done regardless of whether the output will be sent or not.
 
@@ -977,7 +986,12 @@ void FElectraProtronPlayer::FImpl::FDecoderThread::HandleOutputFrame()
 
 			// When sending the output we need to wait until there is room to receive it.
 			// If the output is not to be sent we can handle it quickly without waiting.
-			if (!bSendToOutput || Params.SampleQueueInterface->CanEnqueueAudioSample())
+			if (bSendToOutput)
+			{
+				Params.SampleQueueInterface->AudioSampleWillBeReady();
+			}
+
+			if (!bSendToOutput || (Params.SampleQueueInterface->CanEnqueueAudioSample() && Params.SampleQueueInterface->AreRequiredSamplesReady()))
 			{
 				TSharedPtr<IElectraDecoderAudioOutput, ESPMode::ThreadSafe> AudioDecoderOutput = StaticCastSharedPtr<IElectraDecoderAudioOutput>(CurrentDecoderOutput);
 
diff --git a/Engine/Plugins/Media/ElectraPlayer/Source/ElectraProtron/Private/Player/ElectraProtronPlayerImpl.h b/Engine/Plugins/Media/ElectraPlayer/Source/ElectraProtron/Private/Player/ElectraProtronPlayerImpl.h
--- a/Engine/Plugins/Media/ElectraPlayer/Source/ElectraProtron/Private/Player/ElectraProtronPlayerImpl.h
+++ b/Engine/Plugins/Media/ElectraPlayer/Source/ElectraProtron/Private/Player/ElectraProtronPlayerImpl.h
@@ -87,6 +87,38 @@ public:
 			FScopeLock lock(&TimestampLock);
 			NextExpectedTimestamp.Invalidate();
 			LastHandedOutTimestamp.Invalidate();
+			bVideoSampleWillBeReady = false;
+			bAudioSampleWillBeReady = false;
+		}
+
+		void SetVideoTrackEnabled(bool bInEnabled)
+		{
+			bIsVideoTrackEnabled = bInEnabled;
+		}
+		void SetAudioTrackEnabled(bool bInEnabled)
+		{
+			bIsAudioTrackEnabled = bInEnabled;
+		}
+		bool IsVideoTrackEnabled() const
+		{
+			return bIsVideoTrackEnabled;
+		}
+		bool IsAudioTrackEnabled() const
+		{
+			return bIsAudioTrackEnabled;
+		}
+		void VideoSampleWillBeReady()
+		{
+			bVideoSampleWillBeReady = true;
+		}
+		void AudioSampleWillBeReady()
+		{
+			bAudioSampleWillBeReady = true;
+		}
+		bool AreRequiredSamplesReady() const
+		{
+			return (!bIsVideoTrackEnabled || (bIsVideoTrackEnabled && bVideoSampleWillBeReady)) &&
+				   (!bIsAudioTrackEnabled || (bIsAudioTrackEnabled && bAudioSampleWillBeReady));
 		}
 
 		bool CanEnqueueVideoSample(FTimespan InPTS)
@@ -210,6 +242,10 @@ public:
 		TRange<FTimespan> PlaybackRange;
 		float PlaybackRate = 0.0f;
 		int32 NumVideoFramesToCache = 0;
+		bool bIsVideoTrackEnabled = false;
+		bool bIsAudioTrackEnabled = false;
+		bool bVideoSampleWillBeReady = false;
+		bool bAudioSampleWillBeReady = false;
 	};
 
 
@@ -652,6 +688,7 @@ private:
 		FMediaEvent WorkSignal;
 		bool bTerminate = false;
 		FString LastErrorMessage;
+		double MinLoopSleepTime = 0.0;
 
 		FOpenParam Params;
 		TSharedPtr<FSharedPlayParams, ESPMode::ThreadSafe> SharedPlayParams;
diff --git a/Engine/Plugins/Media/MP4Utilities/Source/MP4Utilities/Private/MP4Utilities.cpp b/Engine/Plugins/Media/MP4Utilities/Source/MP4Utilities/Private/MP4Utilities.cpp
--- a/Engine/Plugins/Media/MP4Utilities/Source/MP4Utilities/Private/MP4Utilities.cpp
+++ b/Engine/Plugins/Media/MP4Utilities/Source/MP4Utilities/Private/MP4Utilities.cpp
@@ -434,7 +434,7 @@ namespace MP4Utilities
 				}
 				else
 				{
-					UE_LOG(LogMP4Utilities, Log, TEXT("Skipping over box `%s` because it is not well known"), *GetPrintableBoxAtom(bi->Type));
+					UE_LOG(LogMP4Utilities, Verbose, TEXT("Skipping over box `%s` because it is not well known"), *GetPrintableBoxAtom(bi->Type));
 					NumRead = InDataReader->ReadData(nullptr, bi->Size-BoxInternalOffset, CurrentOffset+BoxInternalOffset, InCheckCancellationDelegate);
 				}
 			}
