ElectraPlayer: ignoring trailing segment boxes

#rnx

#Synced-CL 51407919

diff --git a/Engine/Plugins/Media/ElectraPlayer/Source/ElectraPlayerRuntime/Private/Runtime/Player/StreamSegmentReaderCommon.cpp b/Engine/Plugins/Media/ElectraPlayer/Source/ElectraPlayerRuntime/Private/Runtime/Player/StreamSegmentReaderCommon.cpp
--- a/Engine/Plugins/Media/ElectraPlayer/Source/ElectraPlayerRuntime/Private/Runtime/Player/StreamSegmentReaderCommon.cpp
+++ b/Engine/Plugins/Media/ElectraPlayer/Source/ElectraPlayerRuntime/Private/Runtime/Player/StreamSegmentReaderCommon.cpp
@@ -1895,6 +1895,7 @@ FStreamSegmentReaderCommon::FStreamHandler::EHandleResult FStreamSegmentReaderCo
 	bool bTimeCheckPassed = InRequest->TimestampVars.Next.bCheck == false;
 	bool bSendMetaData = !!InitSegment->GetMETABox();
 	bool bSendEventData = true;
+	int32 NumMoofs = 0;
 	while(!bDone && !HasErrored() && !HasReadBeenAborted())
 	{
 		if (HasReachedEOF())
@@ -1926,6 +1927,10 @@ FStreamSegmentReaderCommon::FStreamHandler::EHandleResult FStreamSegmentReaderCo
 			}
 			break;
 		}
+		if (HasErrored() || HasReadBeenAborted())
+		{
+			break;
+		}
 
 		TArray<TSharedPtr<MP4Boxes::FMP4BoxBase, ESPMode::ThreadSafe>> ParsedRootBoxes;
 		ParsedRootBoxes.SetNum(RootBoxes.Num());
@@ -1942,6 +1947,11 @@ FStreamSegmentReaderCommon::FStreamHandler::EHandleResult FStreamSegmentReaderCo
 				ParsedRootBoxes[i] = tp.GetBoxTree();
 				switch(ParsedRootBoxes[i]->GetType())
 				{
+					case MP4Utilities::MakeBoxAtom('m','o','o','f'):
+					{
+						++NumMoofs;
+						break;
+					}
 					case MP4Utilities::MakeBoxAtom('m','d','a','t'):
 					{
 						LastByteOffsetOfMDAT = ParsedRootBoxes[i]->GetBoxFileOffset() + ParsedRootBoxes[i]->GetBoxSize();
@@ -1996,7 +2006,16 @@ FStreamSegmentReaderCommon::FStreamHandler::EHandleResult FStreamSegmentReaderCo
 		}
 		if (LastByteOffsetOfMDAT < 0)
 		{
-			return SetError(FString::Printf(TEXT("`mdat` box not found in segment \"%s\""), *RequestURL), INTERNAL_SEG_ERROR_BAD_MEDIA_SEGMENT);
+			if (NumMoofs == 0)
+			{
+				return SetError(FString::Printf(TEXT("`mdat` box not found in segment \"%s\""), *RequestURL), INTERNAL_SEG_ERROR_BAD_MEDIA_SEGMENT);
+			}
+			else
+			{
+				// We probably encountered trailing boxes like `mfra` or similar.
+				bDone = true;
+				break;
+			}
 		}
 
 		struct FSelectedTrackParams
