ElectraPlayer: Reporting error on unsupported HLS DRM scheme

#rnx

#Synced-CL 51452399

diff --git a/Engine/Plugins/Media/ElectraPlayer/Source/ElectraPlayerRuntime/Private/Runtime/Player/AdaptiveStreamingPlayer.cpp b/Engine/Plugins/Media/ElectraPlayer/Source/ElectraPlayerRuntime/Private/Runtime/Player/AdaptiveStreamingPlayer.cpp
--- a/Engine/Plugins/Media/ElectraPlayer/Source/ElectraPlayerRuntime/Private/Runtime/Player/AdaptiveStreamingPlayer.cpp
+++ b/Engine/Plugins/Media/ElectraPlayer/Source/ElectraPlayerRuntime/Private/Runtime/Player/AdaptiveStreamingPlayer.cpp
@@ -2691,13 +2691,25 @@ void FAdaptiveStreamingPlayer::InternalHandlePendingStartRequest(const FTimeValu
 						// Throw a playback error for now.
 						FErrorDetail err;
 						err.SetFacility(Facility::EFacility::Player);
-						err.SetMessage("Could not locate the playback period to begin playback at.");
+						err.SetMessage(TEXT("Could not locate the playback period to begin playback at."));
 						err.SetCode(INTERR_COULD_NOT_LOCATE_START_PERIOD);
 						PostError(err);
 						InitialPlayPeriod.Reset();
 						PendingStartRequest.Reset();
 						break;
 					}
+					case IManifest::FResult::EType::Error:
+					default:
+					{
+						FErrorDetail err;
+						err.SetFacility(Facility::EFacility::Player);
+						err.SetMessage(FString::Printf(TEXT("Error selecting the initial set of streams (%s)."), *Result.GetErrorDetail().GetMessage()));
+						err.SetCode(INTERR_GENERAL_FAILURE);
+						PostError(err);
+						PendingStartRequest.Reset();
+						InitialPlayPeriod.Reset();
+						break;
+					}
 				}
 			}
 			// Do we have a starting play period now?
@@ -2854,7 +2866,7 @@ void FAdaptiveStreamingPlayer::InternalHandlePendingStartRequest(const FTimeValu
 						{
 							FErrorDetail err;
 							err.SetFacility(Facility::EFacility::Player);
-							err.SetMessage("All streams have failed. There is nothing to play any more.");
+							err.SetMessage(TEXT("All streams have failed. There is nothing to play any more."));
 							err.SetCode(INTERR_ALL_STREAMS_HAVE_FAILED);
 							PostError(err);
 							PendingStartRequest.Reset();
@@ -2985,13 +2997,25 @@ void FAdaptiveStreamingPlayer::InternalHandlePendingStartRequest(const FTimeValu
 
 									FErrorDetail err;
 									err.SetFacility(Facility::EFacility::Player);
-									err.SetMessage("Could not locate the stream segment to begin playback at.");
+									err.SetMessage(TEXT("Could not locate the stream segment to begin playback at."));
 									err.SetCode(INTERR_COULD_NOT_LOCATE_START_SEGMENT);
 									PostError(err);
 									PendingStartRequest.Reset();
 									InitialPlayPeriod.Reset();
 									break;
 								}
+								case IManifest::FResult::EType::Error:
+								default:
+								{
+									FErrorDetail err;
+									err.SetFacility(Facility::EFacility::Player);
+									err.SetMessage(FString::Printf(TEXT("Error selecting the initial set of streams (%s)."), *Result.GetErrorDetail().GetMessage()));
+									err.SetCode(INTERR_GENERAL_FAILURE);
+									PostError(err);
+									PendingStartRequest.Reset();
+									InitialPlayPeriod.Reset();
+									break;
+								}
 							}
 						}
 						break;
@@ -3319,7 +3343,7 @@ void FAdaptiveStreamingPlayer::HandlePendingMediaSegmentRequests()
 				{
 					FErrorDetail err;
 					err.SetFacility(Facility::EFacility::Player);
-					err.SetMessage("Could not locate startover period for track switch");
+					err.SetMessage(TEXT("Could not locate startover period for track switch"));
 					err.SetCode(INTERR_COULD_NOT_LOCATE_START_PERIOD);
 					PostError(err);
 					break;
@@ -3398,7 +3422,7 @@ void FAdaptiveStreamingPlayer::HandlePendingMediaSegmentRequests()
 			{
 				FErrorDetail err;
 				err.SetFacility(Facility::EFacility::Player);
-				err.SetMessage("All streams have failed. There is nothing to play any more.");
+				err.SetMessage(TEXT("All streams have failed. There is nothing to play any more."));
 				err.SetCode(INTERR_ALL_STREAMS_HAVE_FAILED);
 				PostError(err);
 				break;
@@ -3511,6 +3535,7 @@ void FAdaptiveStreamingPlayer::HandlePendingMediaSegmentRequests()
 						case IManifest::FResult::EType::NotFound:
 						case IManifest::FResult::EType::BeforeStart:
 						case IManifest::FResult::EType::NotLoaded:
+						case IManifest::FResult::EType::Error:
 						{
 							// Get the dependent stream types, if any.
 							// This is mainly for multiplexed streams like an .mp4 at this point.
@@ -3550,11 +3575,13 @@ void FAdaptiveStreamingPlayer::HandlePendingMediaSegmentRequests()
 				}
 				case IManifest::FResult::EType::BeforeStart:
 				case IManifest::FResult::EType::NotLoaded:
+				case IManifest::FResult::EType::Error:
+				default:
 				{
 					// Throw a playback error for now.
 					FErrorDetail err;
 					err.SetFacility(Facility::EFacility::Player);
-					err.SetMessage("Could not locate next segment");
+					err.SetMessage(FString::Printf(TEXT("Could not locate next segment (%s)"), *Result.GetErrorDetail().GetMessage()));
 					err.SetCode(INTERR_FRAGMENT_NOT_AVAILABLE);
 					PostError(err);
 					break;
@@ -3606,7 +3633,7 @@ void FAdaptiveStreamingPlayer::HandlePendingMediaSegmentRequests()
 				{
 					FErrorDetail err;
 					err.SetFacility(Facility::EFacility::Player);
-					err.SetMessage("Failed to add stream segment request to reader");
+					err.SetMessage(TEXT("Failed to add stream segment request to reader"));
 					err.SetCode(INTERR_FRAGMENT_READER_REQUEST);
 					PostError(err);
 					break;
@@ -4381,7 +4408,7 @@ void FAdaptiveStreamingPlayer::InternalStartAt(const FSeekParam& NewPosition, co
 	{
 		FErrorDetail err;
 		err.SetFacility(Facility::EFacility::Player);
-		err.SetMessage("Failed to create stream reader");
+		err.SetMessage(TEXT("Failed to create stream reader"));
 		err.SetCode(INTERR_CREATE_FRAGMENT_READER);
 		PostError(err);
 		return;
@@ -4836,7 +4863,7 @@ void FAdaptiveStreamingPlayer::InternalRebuffer()
 	{
 		FErrorDetail err;
 		err.SetFacility(Facility::EFacility::Player);
-		err.SetMessage("Rebuffering is set to generate a playback error");
+		err.SetMessage(TEXT("Rebuffering is set to generate a playback error"));
 		err.SetCode(INTERR_REBUFFER_SHALL_THROW_ERROR);
 		PostError(err);
 	}
diff --git a/Engine/Plugins/Media/ElectraPlayer/Source/ElectraPlayerRuntime/Private/Runtime/Player/AdaptiveStreamingPlayerInternal.h b/Engine/Plugins/Media/ElectraPlayer/Source/ElectraPlayerRuntime/Private/Runtime/Player/AdaptiveStreamingPlayerInternal.h
--- a/Engine/Plugins/Media/ElectraPlayer/Source/ElectraPlayerRuntime/Private/Runtime/Player/AdaptiveStreamingPlayerInternal.h
+++ b/Engine/Plugins/Media/ElectraPlayer/Source/ElectraPlayerRuntime/Private/Runtime/Player/AdaptiveStreamingPlayerInternal.h
@@ -49,6 +49,7 @@
 #define INTERR_UNSUPPORTED_CODEC				5
 #define INTERR_CODEC_CHANGE_NOT_SUPPORTED		6
 #define INTERR_NO_STREAM_INFORMATION			7
+#define INTERR_GENERAL_FAILURE					8
 #define INTERR_FRAGMENT_NOT_AVAILABLE			0x101
 #define INTERR_FRAGMENT_READER_REQUEST			0x102
 #define INTERR_CREATE_FRAGMENT_READER			0x103
diff --git a/Engine/Plugins/Media/ElectraPlayer/Source/ElectraPlayerRuntime/Private/Runtime/Player/HLS/PlaylistHLS.cpp b/Engine/Plugins/Media/ElectraPlayer/Source/ElectraPlayerRuntime/Private/Runtime/Player/HLS/PlaylistHLS.cpp
--- a/Engine/Plugins/Media/ElectraPlayer/Source/ElectraPlayerRuntime/Private/Runtime/Player/HLS/PlaylistHLS.cpp
+++ b/Engine/Plugins/Media/ElectraPlayer/Source/ElectraPlayerRuntime/Private/Runtime/Player/HLS/PlaylistHLS.cpp
@@ -3115,7 +3115,7 @@ IManifest::FResult FActiveHLSPlaylist::FPlayPeriod::GetStartingSegment(TSharedPt
 		}
 		else if (SegRes == FTimelineMediaAsset::ESegSearchResult::UnsupportedDRM)
 		{
-			return IManifest::FResult(IManifest::FResult::EType::NotFound).SetErrorDetail(TimelineMediaAsset->GetLastError());
+			return IManifest::FResult(IManifest::FResult::EType::Error).SetErrorDetail(TimelineMediaAsset->GetLastError());
 		}
 		else if (SegRes == FTimelineMediaAsset::ESegSearchResult::PastEOS || SegRes == FTimelineMediaAsset::ESegSearchResult::Ended)
 		{
@@ -3366,7 +3366,7 @@ IManifest::FResult FActiveHLSPlaylist::FPlayPeriod::GetSegment(TSharedPtrTS<IStr
 	}
 	else if (SegRes == FTimelineMediaAsset::ESegSearchResult::UnsupportedDRM)
 	{
-		return IManifest::FResult(IManifest::FResult::EType::NotFound).SetErrorDetail(TimelineMediaAsset->GetLastError());
+		return IManifest::FResult(IManifest::FResult::EType::Error).SetErrorDetail(TimelineMediaAsset->GetLastError());
 	}
 	else if (SegRes == FTimelineMediaAsset::ESegSearchResult::PastEOS || SegRes == FTimelineMediaAsset::ESegSearchResult::Ended)
 	{
diff --git a/Engine/Plugins/Media/ElectraPlayer/Source/ElectraPlayerRuntime/Private/Runtime/Player/Manifest.h b/Engine/Plugins/Media/ElectraPlayer/Source/ElectraPlayerRuntime/Private/Runtime/Player/Manifest.h
--- a/Engine/Plugins/Media/ElectraPlayer/Source/ElectraPlayerRuntime/Private/Runtime/Player/Manifest.h
+++ b/Engine/Plugins/Media/ElectraPlayer/Source/ElectraPlayerRuntime/Private/Runtime/Player/Manifest.h
@@ -183,6 +183,7 @@ namespace Electra
 				BeforeStart,		//!< Time is before the start time
 				TryAgainLater,		//!< Not found at the moment. Playlist load may be pending
 				NotLoaded,			//!< Not loaded (playlist has not been requested)
+				Error				//!< An unrecoverable error
 			};
 			FResult(EType InType = EType::NotFound)
 				: Type(InType)
@@ -252,6 +253,10 @@ namespace Electra
 					{
 						return TEXT("Not loaded");
 					}
+					case EType::Error:
+					{
+						return TEXT("Error");
+					}
 					default:
 					{
 						return TEXT("undefined");
