Protron: added back an accidentally removed check for pending output

#rnx

#Synced-CL 0

diff --git a/Engine/Plugins/Media/ElectraPlayer/Source/ElectraProtron/Private/Player/ElectraProtronPlayerDecoder.cpp b/Engine/Plugins/Media/ElectraPlayer/Source/ElectraProtron/Private/Player/ElectraProtronPlayerDecoder.cpp
--- a/Engine/Plugins/Media/ElectraPlayer/Source/ElectraProtron/Private/Player/ElectraProtronPlayerDecoder.cpp
+++ b/Engine/Plugins/Media/ElectraPlayer/Source/ElectraProtron/Private/Player/ElectraProtronPlayerDecoder.cpp
@@ -546,6 +546,13 @@ void FElectraProtronPlayer::FImpl::FDecoderThread::DecodeOneFrame()
 		return;
 	}
 
+	// If we are still hanging on to output we could not deliver we will not
+	// decode new data that will likewise not be able to deliver.
+	if (CurrentDecoderOutput.IsValid())
+	{
+		return;
+	}
+
 	// If the end of the stream has been reached we are done.
 	if (bReachedEnd)
 	{
