ElectraPlayer: Early checking if all full segment encrypted bytes have already been decrypted.

#rnx

#Synced-CL 51619714

diff --git a/Engine/Plugins/Media/ElectraPlayer/Source/ElectraPlayerRuntime/Private/Runtime/Player/StreamSegmentReaderCommon.cpp b/Engine/Plugins/Media/ElectraPlayer/Source/ElectraPlayerRuntime/Private/Runtime/Player/StreamSegmentReaderCommon.cpp
--- a/Engine/Plugins/Media/ElectraPlayer/Source/ElectraPlayerRuntime/Private/Runtime/Player/StreamSegmentReaderCommon.cpp
+++ b/Engine/Plugins/Media/ElectraPlayer/Source/ElectraPlayerRuntime/Private/Runtime/Player/StreamSegmentReaderCommon.cpp
@@ -5243,6 +5243,12 @@ bool FStreamSegmentReaderCommon::FStreamHandler::FReadBuffer::WaitUntilSizeAvail
 			DecryptedDataBuffer.Abort();
 			return true;
 		}
+		// Have all data decrypted already?
+		if (DecryptedDataBuffer.GetEOD() && !BlockDecrypterHandle)
+		{
+			return true;
+		}
+
 		// Decrypt
 		FScopeLock lock(ReceiveBuffer->GetLock());
 		int64 SourceSizeAvail = ReceiveBuffer->GetLinearReadSize();
