Fixed PVS warning

#rnx

#Synced-CL 51619714

diff --git a/Engine/Plugins/Media/ElectraPlayer/Source/ElectraPlayerRuntime/Private/Runtime/Player/StreamSegmentReaderCommon.cpp b/Engine/Plugins/Media/ElectraPlayer/Source/ElectraPlayerRuntime/Private/Runtime/Player/StreamSegmentReaderCommon.cpp
--- a/Engine/Plugins/Media/ElectraPlayer/Source/ElectraPlayerRuntime/Private/Runtime/Player/StreamSegmentReaderCommon.cpp
+++ b/Engine/Plugins/Media/ElectraPlayer/Source/ElectraPlayerRuntime/Private/Runtime/Player/StreamSegmentReaderCommon.cpp
@@ -5205,7 +5205,7 @@ bool FStreamSegmentReaderCommon::FStreamHandler::FReadBuffer::WaitUntilSizeAvail
 				return false;
 			}
 			bIsDecrypterGood = BlockDecrypter->GetState() == ElectraCDM::ECDMState::Ready;
-			if (bIsDecrypterGood)
+			if (bIsDecrypterGood && BlockDecrypterHandle)
 			{
 				ElectraCDM::FMediaCDMSampleInfo si;
 				si.IV = BlockDecrypterIV;
