ElectraPlayer: Adjusted event times for HLS

#rnx

#Synced-CL 51661707

diff --git a/Engine/Plugins/Media/ElectraPlayer/Source/ElectraPlayerRuntime/Private/Runtime/Player/HLS/PlaylistHLS.cpp b/Engine/Plugins/Media/ElectraPlayer/Source/ElectraPlayerRuntime/Private/Runtime/Player/HLS/PlaylistHLS.cpp
--- a/Engine/Plugins/Media/ElectraPlayer/Source/ElectraPlayerRuntime/Private/Runtime/Player/HLS/PlaylistHLS.cpp
+++ b/Engine/Plugins/Media/ElectraPlayer/Source/ElectraPlayerRuntime/Private/Runtime/Player/HLS/PlaylistHLS.cpp
@@ -2084,24 +2084,7 @@ void FActiveHLSPlaylist::FTimelineMediaAsset::TransformIntoReplayEvent()
 	bool bStaticStart = false;
 	bool bDynamicStart = false;
 
-	// Get the replay base start time.
 	const TArray<FURL_RFC3986::FQueryParam>& frgs(ReplayEventParams.GetUrlFragmentParams());
-	for(int32 i=0,iMax=frgs.Num(); i<iMax; ++i)
-	{
-		if (frgs[i].Name.Equals(CustomOptions::Custom_EpicUTCNow))
-		{
-			if (!UnixEpoch::ParseFloatString(ReplayEvent.BaseStartTime, frgs[i].Value))
-			{
-				UE_LOG(LogElectraPlayer, Warning, TEXT("Could not parse the custom UTC now value, disabling replay event."));
-				return;
-			}
-			break;
-		}
-		else if (frgs[i].Name.Equals(CustomOptions::Custom_EpicUTCUrl))
-		{
-			UE_LOG(LogElectraPlayer, Warning, TEXT("Fetching UTC time from a custom URL for a replay event is not supported with HLS."));
-		}
-	}
 	// Which type of replay event?
 	for(int32 i=0,iMax=frgs.Num(); i<iMax; ++i)
 	{
@@ -2126,6 +2109,24 @@ void FActiveHLSPlaylist::FTimelineMediaAsset::TransformIntoReplayEvent()
 		return;
 	}
 
+	// Get the replay base start time.
+	for(int32 i=0,iMax=frgs.Num(); i<iMax; ++i)
+	{
+		if (frgs[i].Name.Equals(CustomOptions::Custom_EpicUTCNow))
+		{
+			if (!UnixEpoch::ParseFloatString(ReplayEvent.BaseStartTime, frgs[i].Value))
+			{
+				UE_LOG(LogElectraPlayer, Warning, TEXT("Could not parse the custom UTC now value, disabling replay event."));
+				return;
+			}
+			break;
+		}
+		else if (frgs[i].Name.Equals(CustomOptions::Custom_EpicUTCUrl))
+		{
+			UE_LOG(LogElectraPlayer, Warning, TEXT("Fetching UTC time from a custom URL for a replay event is not supported with HLS."));
+		}
+	}
+
 	TArray<FString> Params;
 	if (!Arg.IsEmpty())
 	{
@@ -2134,6 +2135,7 @@ void FActiveHLSPlaylist::FTimelineMediaAsset::TransformIntoReplayEvent()
 	}
 	if (Params.Num())
 	{
+		FTimeValue Start;
 		ReplayEvent.bIsReplay = true;
 		ReplayEvent.bIsStaticEvent = bStaticStart;
 		if (ReplayEvent.BaseStartTime.IsValid())
@@ -2146,22 +2148,29 @@ void FActiveHLSPlaylist::FTimelineMediaAsset::TransformIntoReplayEvent()
 		if (Time.StartsWith(TEXT("now")))
 		{
 			Time.RightChopInline(3);
+			Start = PlayerSessionServices->GetSynchronizedUTCTime()->GetTime();
 		}
 		if (!Time.IsEmpty())
 		{
 			FTimeValue Offset = FTimeValue().SetFromTimeFraction(FTimeFraction().SetFromFloatString(Time));
 			if (Offset.IsValid())
 			{
-				if (ReplayEvent.BaseStartTime.IsValid())
+				if (Start.IsValid())
 				{
-					ReplayEvent.BaseStartTime += Offset;
+					Start += Offset;
 				}
 				else
 				{
-					ReplayEvent.BaseStartTime = Offset;
+					Start = Offset;
 				}
 			}
 		}
+		if (!Start.IsValid())
+		{
+			ReplayEvent.Reset();
+			return;
+		}
+		ReplayEvent.BaseStartTime = Start;
 		ReplayEvent.OriginalBaseTimeOffset = BaseTimeOffset;
 		BaseTimeOffset = ReplayEvent.BaseStartTime;
 		if (Params.Num() > 1)
diff --git a/Engine/Plugins/Media/ElectraPlayer/Source/ElectraPlayerRuntime/Private/Runtime/Player/HLS/PlaylistHLS.h b/Engine/Plugins/Media/ElectraPlayer/Source/ElectraPlayerRuntime/Private/Runtime/Player/HLS/PlaylistHLS.h
--- a/Engine/Plugins/Media/ElectraPlayer/Source/ElectraPlayerRuntime/Private/Runtime/Player/HLS/PlaylistHLS.h
+++ b/Engine/Plugins/Media/ElectraPlayer/Source/ElectraPlayerRuntime/Private/Runtime/Player/HLS/PlaylistHLS.h
@@ -721,6 +721,14 @@ private:
 			FTimeValue SuggestedPresentationDelay;
 			bool bIsReplay = false;
 			bool bIsStaticEvent = false;
+			void Reset()
+			{
+				BaseStartTime.SetToInvalid();
+				OriginalBaseTimeOffset.SetToInvalid();
+				SuggestedPresentationDelay.SetToInvalid();
+				bIsReplay = false;
+				bIsStaticEvent = false;
+			}
 		};
 		FReplayEvent ReplayEvent;
 	};
