ElectraPlayer: fixed incorrect full segment decrypter setup

#rnx

#Synced-CL 51755220

diff --git a/Engine/Plugins/Media/ElectraPlayer/Source/ElectraPlayerRuntime/Private/Runtime/Player/StreamSegmentReaderCommon.cpp b/Engine/Plugins/Media/ElectraPlayer/Source/ElectraPlayerRuntime/Private/Runtime/Player/StreamSegmentReaderCommon.cpp
--- a/Engine/Plugins/Media/ElectraPlayer/Source/ElectraPlayerRuntime/Private/Runtime/Player/StreamSegmentReaderCommon.cpp
+++ b/Engine/Plugins/Media/ElectraPlayer/Source/ElectraPlayerRuntime/Private/Runtime/Player/StreamSegmentReaderCommon.cpp
@@ -5205,7 +5205,7 @@ bool FStreamSegmentReaderCommon::FStreamHandler::FReadBuffer::WaitUntilSizeAvail
 				return false;
 			}
 			bIsDecrypterGood = BlockDecrypter->GetState() == ElectraCDM::ECDMState::Ready;
-			if (bIsDecrypterGood && BlockDecrypterHandle)
+			if (bIsDecrypterGood && !BlockDecrypterHandle)
 			{
 				ElectraCDM::FMediaCDMSampleInfo si;
 				si.IV = BlockDecrypterIV;
@@ -5215,7 +5215,11 @@ bool FStreamSegmentReaderCommon::FStreamHandler::FReadBuffer::WaitUntilSizeAvail
 				{
 					bIsDecrypterGood = false;
 				}
-				BlockDecrypterBlockSize = BlockDecrypterHandle->BlockSize;
+				BlockDecrypterBlockSize = BlockDecrypterHandle ? BlockDecrypterHandle->BlockSize : 0;
+				if (BlockDecrypterBlockSize <= 0)
+				{
+					bIsDecrypterGood = false;
+				}
 			}
 			bDecrypterReady = true;
 		}
@@ -5257,9 +5261,8 @@ bool FStreamSegmentReaderCommon::FStreamHandler::FReadBuffer::WaitUntilSizeAvail
 		{
 			int64 DecryptedSizeAvail = DecryptedDataBuffer.GetLinearReadSize();
 			int64 NumNew = Needed - DecryptedSizeAvail;
-			check(NumNew % BlockDecrypterBlockSize == 0);
 			DecryptedDataBuffer.EnlargeTo(SourceSizeAvail);
-			if (NumNew > 0)
+			if (NumNew > 0 && BlockDecrypterHandle)
 			{
 				uint8* NewDataToDecryptPtr = DecryptedDataBuffer.GetLinearWriteData(NumNew);
 				FMemory::Memcpy(NewDataToDecryptPtr, ReceiveBuffer->GetLinearReadData() + DecryptedSizeAvail, NumNew);
